Fix MultiUserServer compilation error in shipping configuration

diff --git a/Engine/Plugins/Developer/Concert/ConcertMain/Source/Concert/Private/ConcertClient.cpp b/Engine/Plugins/Developer/Concert/ConcertMain/Source/Concert/Private/ConcertClient.cpp
--- a/Engine/Plugins/Developer/Concert/ConcertMain/Source/Concert/Private/ConcertClient.cpp
+++ b/Engine/Plugins/Developer/Concert/ConcertMain/Source/Concert/Private/ConcertClient.cpp
@@ -18,6 +18,19 @@
 
 #define LOCTEXT_NAMESPACE "ConcertClient"
 
+
+namespace ConcertUtil
+{
+	const FLogCategoryBase* GetLogConcertPtr()
+	{
+	#if NO_LOGGING
+		return nullptr;
+	#else
+		return &LogConcert;
+	#endif
+	}
+}
+
 class FConcertAutoConnection
 {
 public:
@@ -65,7 +78,7 @@ private:
 		NotificationConfig.ProgressText = FText::Format(LOCTEXT("LookingForServer", "Looking for Server '{0}'..."), FText::FromString(Settings->DefaultServerURL));
 		NotificationConfig.bIsHeadless = Settings->bIsHeadless;
 		NotificationConfig.bCanCancel = true;
-		NotificationConfig.LogCategory = &LogConcert;
+		NotificationConfig.LogCategory = ConcertUtil::GetLogConcertPtr();
 		return MakeUnique<FAsyncTaskNotification>(NotificationConfig);
 	}
 
@@ -316,7 +329,7 @@ public:
 			FAsyncTaskNotificationConfig NotificationConfig;
 			NotificationConfig.TitleText = Config.PendingTitleText.Get(FText::GetEmpty());
 			NotificationConfig.bIsHeadless = Client->GetConfiguration()->bIsHeadless;
-			NotificationConfig.LogCategory = &LogConcert;
+			NotificationConfig.LogCategory = ConcertUtil::GetLogConcertPtr();
 			Notification = MakeUnique<FAsyncTaskNotification>(NotificationConfig);
 		}
 		Notification->SetCanCancel(TAttribute<bool>(this, &FConcertPendingConnection::CanCancel));
@@ -770,7 +783,7 @@ TFuture<EConcertResponseCode> FConcertClient::ArchiveSession(const FGuid& Server
 	NotificationConfig.bIsHeadless = Settings->bIsHeadless;
 	NotificationConfig.bKeepOpenOnFailure = true;
 	NotificationConfig.TitleText = LOCTEXT("ArchivingSession", "Archiving Session...");
-	NotificationConfig.LogCategory = &LogConcert;
+	NotificationConfig.LogCategory = ConcertUtil::GetLogConcertPtr();
 
 	FAsyncTaskNotification Notification(NotificationConfig);
 
@@ -803,7 +816,7 @@ TFuture<EConcertResponseCode> FConcertClient::RenameSession(const FGuid& ServerA
 	NotificationConfig.bIsHeadless = Settings->bIsHeadless;
 	NotificationConfig.bKeepOpenOnFailure = true;
 	NotificationConfig.TitleText = LOCTEXT("RenamingSession", "Renaming Session...");
-	NotificationConfig.LogCategory = &LogConcert;
+	NotificationConfig.LogCategory = ConcertUtil::GetLogConcertPtr();
 
 	FAsyncTaskNotification Notification(NotificationConfig);
 
@@ -835,7 +848,7 @@ TFuture<EConcertResponseCode> FConcertClient::DeleteSession(const FGuid& ServerA
 	NotificationConfig.bIsHeadless = Settings->bIsHeadless;
 	NotificationConfig.bKeepOpenOnFailure = true;
 	NotificationConfig.TitleText = LOCTEXT("DeletingSession", "Deleting Session...");
-	NotificationConfig.LogCategory = &LogConcert;
+	NotificationConfig.LogCategory = ConcertUtil::GetLogConcertPtr();
 
 	FAsyncTaskNotification Notification(NotificationConfig);
 
@@ -1025,7 +1038,7 @@ TFuture<EConcertResponseCode> FConcertClient::InternalRestoreSession(const FGuid
 		NotificationConfig.bIsHeadless = Settings->bIsHeadless;
 		NotificationConfig.bKeepOpenOnFailure = true;
 		NotificationConfig.TitleText = LOCTEXT("RestoringSession", "Restoring Session...");
-		NotificationConfig.LogCategory = &LogConcert;
+		NotificationConfig.LogCategory = ConcertUtil::GetLogConcertPtr();
 		Notification = MakeUnique<FAsyncTaskNotification>(NotificationConfig);
 	}
 
