Added friendly name to the custom version information sent to the server

This allows the server to provide better error messages, as it only natively understands the core set of custom versions, and not versions added by plugins on the client (where it would have previously just shown you the version GUID).

#rnx

diff --git a/Engine/Plugins/Developer/Concert/ConcertMain/Source/Concert/Private/ConcertVersion.cpp b/Engine/Plugins/Developer/Concert/ConcertMain/Source/Concert/Private/ConcertVersion.cpp
--- a/Engine/Plugins/Developer/Concert/ConcertMain/Source/Concert/Private/ConcertVersion.cpp
+++ b/Engine/Plugins/Developer/Concert/ConcertMain/Source/Concert/Private/ConcertVersion.cpp
@@ -78,6 +78,7 @@ bool FConcertEngineVersionInfo::Validate(const FConcertEngineVersionInfo& InOthe
 
 void FConcertCustomVersionInfo::Initialize(const FCustomVersion& InVersion)
 {
+	FriendlyName = InVersion.GetFriendlyName();
 	Key = InVersion.Key;
 	Version = InVersion.Version;
 }
@@ -85,8 +86,7 @@ void FConcertCustomVersionInfo::Initialize(const FCustomVersion& InVersion)
 bool FConcertCustomVersionInfo::Validate(const FConcertCustomVersionInfo& InOther, const EConcertVersionValidationMode InValidationMode, FText* OutFailureReason) const
 {
 	check(Key == InOther.Key);
-	const FCustomVersion* EngineCustomVersion = FCustomVersionContainer::GetRegistered().GetVersion(Key);
-	return ConcertVersionUtil::ValidateVersion(Version, InOther.Version, EngineCustomVersion ? FText::AsCultureInvariant(EngineCustomVersion->GetFriendlyName().ToString()) : FText::AsCultureInvariant(Key.ToString()), InValidationMode, OutFailureReason);
+	return ConcertVersionUtil::ValidateVersion(Version, InOther.Version, FText::AsCultureInvariant(FriendlyName.IsNone() ? Key.ToString() : FriendlyName.ToString()), InValidationMode, OutFailureReason);
 }
 
 
@@ -125,8 +125,7 @@ bool FConcertSessionVersionInfo::Validate(const FConcertSessionVersionInfo& InOt
 		{
 			if (OutFailureReason)
 			{
-				const FCustomVersion* EngineCustomVersion = FCustomVersionContainer::GetRegistered().GetVersion(CustomVersion.Key);
-				*OutFailureReason = FText::Format(LOCTEXT("Error_MissingVersionFmt", "Invalid version for '{0}' (expected '{1}', got '<none>')"), EngineCustomVersion ? FText::AsCultureInvariant(EngineCustomVersion->GetFriendlyName().ToString()) : FText::AsCultureInvariant(CustomVersion.Key.ToString()), CustomVersion.Version);
+				*OutFailureReason = FText::Format(LOCTEXT("Error_MissingVersionFmt", "Invalid version for '{0}' (expected '{1}', got '<none>'). Do you have a required plugin disabled?"), FText::AsCultureInvariant(CustomVersion.FriendlyName.IsNone() ? CustomVersion.Key.ToString() : CustomVersion.FriendlyName.ToString()), CustomVersion.Version);
 			}
 			return false;
 		}
@@ -153,8 +152,7 @@ bool FConcertSessionVersionInfo::Validate(const FConcertSessionVersionInfo& InOt
 
 				if (!CustomVersion)
 				{
-					const FCustomVersion* EngineCustomVersion = FCustomVersionContainer::GetRegistered().GetVersion(OtherCustomVersion.Key);
-					*OutFailureReason = FText::Format(LOCTEXT("Error_ExtraCustomVersionFmt", "Invalid version for '{0}' (expected '<none>', got '{1}')"), EngineCustomVersion ? FText::AsCultureInvariant(EngineCustomVersion->GetFriendlyName().ToString()) : FText::AsCultureInvariant(OtherCustomVersion.Key.ToString()), OtherCustomVersion.Version);
+					*OutFailureReason = FText::Format(LOCTEXT("Error_ExtraCustomVersionFmt", "Invalid version for '{0}' (expected '<none>', got '{1}'). Do you have an extra plugin enabled?"), FText::AsCultureInvariant(OtherCustomVersion.FriendlyName.IsNone() ? OtherCustomVersion.Key.ToString() : OtherCustomVersion.FriendlyName.ToString()), OtherCustomVersion.Version);
 					break;
 				}
 			}
diff --git a/Engine/Plugins/Developer/Concert/ConcertMain/Source/Concert/Public/ConcertVersion.h b/Engine/Plugins/Developer/Concert/ConcertMain/Source/Concert/Public/ConcertVersion.h
--- a/Engine/Plugins/Developer/Concert/ConcertMain/Source/Concert/Public/ConcertVersion.h
+++ b/Engine/Plugins/Developer/Concert/ConcertMain/Source/Concert/Public/ConcertVersion.h
@@ -79,6 +79,10 @@ struct FConcertCustomVersionInfo
 	/** Validate this version info against another */
 	CONCERT_API bool Validate(const FConcertCustomVersionInfo& InOther, const EConcertVersionValidationMode InValidationMode, FText* OutFailureReason = nullptr) const;
 
+	/** Friendly name of the version */
+	UPROPERTY()
+	FName FriendlyName;
+
 	/** Unique custom key */
 	UPROPERTY()
 	FGuid Key;
diff --git a/Engine/Plugins/Developer/Concert/ConcertSync/ConcertSyncCore/Source/ConcertSyncCore/Private/ConcertSyncArchives.cpp b/Engine/Plugins/Developer/Concert/ConcertSync/ConcertSyncCore/Source/ConcertSyncCore/Private/ConcertSyncArchives.cpp
--- a/Engine/Plugins/Developer/Concert/ConcertSync/ConcertSyncCore/Source/ConcertSyncCore/Private/ConcertSyncArchives.cpp
+++ b/Engine/Plugins/Developer/Concert/ConcertSync/ConcertSyncCore/Source/ConcertSyncCore/Private/ConcertSyncArchives.cpp
@@ -181,7 +181,7 @@ FConcertSyncObjectReader::FConcertSyncObjectReader(const FConcertLocalIdentifier
 		FCustomVersionContainer EngineCustomVersions;
 		for (const FConcertCustomVersionInfo& CustomVersion : InVersionInfo->CustomVersions)
 		{
-			EngineCustomVersions.SetVersion(CustomVersion.Key, CustomVersion.Version, NAME_None);
+			EngineCustomVersions.SetVersion(CustomVersion.Key, CustomVersion.Version, CustomVersion.FriendlyName);
 		}
 		SetCustomVersions(EngineCustomVersions);
 	}
