Fixed FName instances deserialized by Concert from splitting what is supposed to be a literal name

diff --git a/Engine/Plugins/Developer/Concert/ConcertMain/Source/ConcertTransport/Private/IdentifierTable/ConcertIdentifierTable.cpp b/Engine/Plugins/Developer/Concert/ConcertMain/Source/ConcertTransport/Private/IdentifierTable/ConcertIdentifierTable.cpp
--- a/Engine/Plugins/Developer/Concert/ConcertMain/Source/ConcertTransport/Private/IdentifierTable/ConcertIdentifierTable.cpp
+++ b/Engine/Plugins/Developer/Concert/ConcertMain/Source/ConcertTransport/Private/IdentifierTable/ConcertIdentifierTable.cpp
@@ -16,7 +16,7 @@ int32 FConcertLocalIdentifierTable::MapName(const FName& InName)
 
 	const int32 NewIndex = MappedNames.Num();
 	FName& MappedName = MappedNames.Add_GetRef(InName);
-	MappedName.SetNumber(0); // Clear the number from the version we store
+	MappedName.SetNumber(NAME_NO_NUMBER_INTERNAL); // Clear the number from the version we store
 	return NameToMappedIndex.Add(MappedName, NewIndex);
 }
 
@@ -59,7 +59,7 @@ void FConcertLocalIdentifierTable::SetState(const FConcertLocalIdentifierState&
 	NameToMappedIndex.Reserve(InState.MappedNames.Num());
 	for (const FString& MappedNameStr : InState.MappedNames)
 	{
-		const FName MappedName = FName(*MappedNameStr);
+		const FName MappedName = FName(*MappedNameStr, NAME_NO_NUMBER_INTERNAL, FNAME_Add, /*bSplitName*/false);
 		const int32 NewIndex = MappedNames.Num();
 		MappedNames.Add(MappedName);
 		NameToMappedIndex.Add(MappedName, NewIndex);
diff --git a/Engine/Plugins/Developer/Concert/ConcertMain/Source/ConcertTransport/Private/IdentifierTable/ConcertTransportArchives.cpp b/Engine/Plugins/Developer/Concert/ConcertMain/Source/ConcertTransport/Private/IdentifierTable/ConcertTransportArchives.cpp
--- a/Engine/Plugins/Developer/Concert/ConcertMain/Source/ConcertTransport/Private/IdentifierTable/ConcertTransportArchives.cpp
+++ b/Engine/Plugins/Developer/Concert/ConcertMain/Source/ConcertTransport/Private/IdentifierTable/ConcertTransportArchives.cpp
@@ -94,7 +94,7 @@ FArchive& FConcertIdentifierReader::operator<<(FName& Name)
 		{
 			FString PlainString;
 			*this << PlainString;
-			Name = FName(*PlainString);
+			Name = FName(*PlainString, NAME_NO_NUMBER_INTERNAL, FNAME_Add, /*bSplitName*/false);
 		}
 		break;
 
diff --git a/Engine/Source/Runtime/Core/Private/UObject/UnrealNames.cpp b/Engine/Source/Runtime/Core/Private/UObject/UnrealNames.cpp
--- a/Engine/Source/Runtime/Core/Private/UObject/UnrealNames.cpp
+++ b/Engine/Source/Runtime/Core/Private/UObject/UnrealNames.cpp
@@ -494,9 +494,9 @@ FName::FName(const ANSICHAR* Name, EFindName FindType)
  * @param Number Value for the number portion of the name
  * @param FindType Action to take (see EFindName)
  */
-FName::FName( const TCHAR* Name, int32 InNumber, EFindName FindType )
+FName::FName( const TCHAR* Name, int32 InNumber, EFindName FindType, const bool bSplitName )
 {
-	Init(Name, InNumber, FindType);
+	Init(Name, InNumber, FindType, bSplitName);
 }
 
 FName::FName(const FNameEntrySerialized& LoadedEntry)
diff --git a/Engine/Source/Runtime/Core/Public/UObject/NameTypes.h b/Engine/Source/Runtime/Core/Public/UObject/NameTypes.h
--- a/Engine/Source/Runtime/Core/Public/UObject/NameTypes.h
+++ b/Engine/Source/Runtime/Core/Public/UObject/NameTypes.h
@@ -933,8 +933,9 @@ public:
 	 * @param Name Value for the string portion of the name
 	 * @param Number Value for the number portion of the name
 	 * @param FindType Action to take (see EFindName)
+	 * @param bSplitName true if the trailing number should be split from the name when Number == NAME_NO_NUMBER_INTERNAL, or false to always use the name as-is
 	 */
-	FName( const TCHAR* Name, int32 InNumber, EFindName FindType=FNAME_Add );
+	FName( const TCHAR* Name, int32 InNumber, EFindName FindType=FNAME_Add, const bool bSplitName=true );
 
 	/**
 	 * Constructor used by FLinkerLoad when loading its name table; Creates an FName with an instance
