Fixed line spacing in MultiUser Settings

diff --git a/Engine/Plugins/Developer/Concert/ConcertMain/Source/Concert/Public/ConcertSettings.h b/Engine/Plugins/Developer/Concert/ConcertMain/Source/Concert/Public/ConcertSettings.h
--- a/Engine/Plugins/Developer/Concert/ConcertMain/Source/Concert/Public/ConcertSettings.h
+++ b/Engine/Plugins/Developer/Concert/ConcertMain/Source/Concert/Public/ConcertSettings.h
@@ -184,7 +184,7 @@ struct FConcertClientSettings
 	FSoftClassPath DesktopAvatarActorClass;
 
 	/** The VR representation of this editor's user to other connected users */
-	UPROPERTY(config, EditAnywhere, NoClear, Category = "Client Settings", meta = (MetaClass = "ConcertClientVRPresenceActor"))
+	UPROPERTY(config, EditAnywhere, NoClear, Category = "Client Settings", meta = (MetaClass = "ConcertClientVRPresenceActor", DisplayName = "VR Avatar Actor Class"))
 	FSoftClassPath VRAvatarActorClass;
 
 	/** The timespan at which discovered Concert server are considered stale if they haven't answered back */
