Updated MU engine versioning to use the current build version

We've seen the compiled-in engine versions get out-of-sync when dealing with pre-built binaries vs binaries built locally by a programmer, so we've changed the validation to use the current engine version, and to override the changelist with the Build.version data (if available). The changelist in the Build.version data will always be the latest CL# synced in UGS, regardless of whether you're using a pre-built or locally compiled editor.

This change makes the older BaseRevision setting redundant, but we've kept it around for now (set to zero) so that we can hopefully change it to store the real content base revision in the future (in a way that works for projects not using UGS, eg) using pre-built editors from the public launcher).

Since the base revision is now always zero, it's been removed from the UI, and the UI has been updated to show you the engine version of a session instead.

#rnx

diff --git a/Engine/Plugins/Developer/Concert/ConcertMain/Source/Concert/Private/ConcertMessageData.cpp b/Engine/Plugins/Developer/Concert/ConcertMain/Source/Concert/Private/ConcertMessageData.cpp
--- a/Engine/Plugins/Developer/Concert/ConcertMain/Source/Concert/Private/ConcertMessageData.cpp
+++ b/Engine/Plugins/Developer/Concert/ConcertMain/Source/Concert/Private/ConcertMessageData.cpp
@@ -127,9 +127,20 @@ FText FConcertSessionInfo::ToDisplayString() const
 	FTextBuilder TextBuilder;
 	TextBuilder.AppendLineFormat(NSLOCTEXT("ConcertSessionInfo", "SessionId", "Session ID: {0}"), FText::FromString(SessionId.ToString()));
 	TextBuilder.AppendLineFormat(NSLOCTEXT("ConcertSessionInfo", "SessionName", "Session Name: {0}"), FText::FromString(SessionName));
-	TextBuilder.AppendLineFormat(NSLOCTEXT("ConcertSessionInfo", "OwnerUserName", "Owner User Name: {0}"), FText::FromString(OwnerUserName));
+	TextBuilder.AppendLineFormat(NSLOCTEXT("ConcertSessionInfo", "OwnerUserName", "Session Owner: {0}"), FText::FromString(OwnerUserName));
 	TextBuilder.AppendLineFormat(NSLOCTEXT("ConcertSessionInfo", "ProjectName", "Session Project: {0}"), FText::FromString(Settings.ProjectName));
-	TextBuilder.AppendLineFormat(NSLOCTEXT("ConcertSessionInfo", "BaseRevision", "Session Base Revision: {0}"), FText::AsNumber(Settings.BaseRevision, &FNumberFormattingOptions::DefaultNoGrouping()));
+	//TextBuilder.AppendLineFormat(NSLOCTEXT("ConcertSessionInfo", "BaseRevision", "Session Base Revision: {0}"), FText::AsNumber(Settings.BaseRevision, &FNumberFormattingOptions::DefaultNoGrouping()));
+	if (VersionInfos.Num() > 0)
+	{
+		const FConcertSessionVersionInfo& VersionInfo = VersionInfos.Last();
+		TextBuilder.AppendLineFormat(
+			NSLOCTEXT("ConcertSessionInfo", "EngineVersion", "Session Engine Version: {0}.{1}.{2}-{3}"), 
+			FText::AsNumber(VersionInfo.EngineVersion.Major, &FNumberFormattingOptions::DefaultNoGrouping()),
+			FText::AsNumber(VersionInfo.EngineVersion.Minor, &FNumberFormattingOptions::DefaultNoGrouping()),
+			FText::AsNumber(VersionInfo.EngineVersion.Patch, &FNumberFormattingOptions::DefaultNoGrouping()),
+			FText::AsNumber(VersionInfo.EngineVersion.Changelist, &FNumberFormattingOptions::DefaultNoGrouping())
+			);
+	}
 	TextBuilder.AppendLineFormat(NSLOCTEXT("ConcertSessionInfo", "ServerEndpointId", "Server Endpoint ID: {0}"), FText::FromString(ServerEndpointId.ToString()));
 	return TextBuilder.ToText();
 }
diff --git a/Engine/Plugins/Developer/Concert/ConcertMain/Source/Concert/Private/ConcertVersion.cpp b/Engine/Plugins/Developer/Concert/ConcertMain/Source/Concert/Private/ConcertVersion.cpp
--- a/Engine/Plugins/Developer/Concert/ConcertMain/Source/Concert/Private/ConcertVersion.cpp
+++ b/Engine/Plugins/Developer/Concert/ConcertMain/Source/Concert/Private/ConcertVersion.cpp
@@ -2,6 +2,7 @@
 
 #include "ConcertVersion.h"
 #include "Misc/EngineVersion.h"
+#include "Modules/BuildVersion.h"
 #include "Serialization/CustomVersion.h"
 
 #define LOCTEXT_NAMESPACE "ConcertVersion"
@@ -18,7 +19,7 @@ bool ValidateVersion(const int32 InCurrent, const int32 InOther, const FText InV
 		{
 			if (OutFailureReason)
 			{
-				*OutFailureReason = FText::Format(LOCTEXT("Error_InvalidIdenticalVersionFmt", "Invalid version for '{0}' (expected '{1}', got '{2}')"), InVersionDisplayName, InCurrent, InOther);
+				*OutFailureReason = FText::Format(LOCTEXT("Error_InvalidIdenticalVersionFmt", "Invalid version for '{0}' (expected '{1}', got '{2}')"), InVersionDisplayName, FText::AsNumber(InCurrent, &FNumberFormattingOptions::DefaultNoGrouping()), FText::AsNumber(InOther, &FNumberFormattingOptions::DefaultNoGrouping()));
 			}
 			return false;
 		}
@@ -29,7 +30,7 @@ bool ValidateVersion(const int32 InCurrent, const int32 InOther, const FText InV
 		{
 			if (OutFailureReason)
 			{
-				*OutFailureReason = FText::Format(LOCTEXT("Error_InvalidCompatibleVersionFmt", "Invalid version for '{0}' (expected '{1}' or greater, got '{2}')"), InVersionDisplayName, InCurrent, InOther);
+				*OutFailureReason = FText::Format(LOCTEXT("Error_InvalidCompatibleVersionFmt", "Invalid version for '{0}' (expected '{1}' or greater, got '{2}')"), InVersionDisplayName, FText::AsNumber(InCurrent, &FNumberFormattingOptions::DefaultNoGrouping()), FText::AsNumber(InOther, &FNumberFormattingOptions::DefaultNoGrouping()));
 			}
 			return false;
 		}
@@ -93,7 +94,28 @@ bool FConcertCustomVersionInfo::Validate(const FConcertCustomVersionInfo& InOthe
 void FConcertSessionVersionInfo::Initialize()
 {
 	FileVersion.Initialize();
-	CompatibleEngineVersion.Initialize(FEngineVersion::CompatibleWith());
+	EngineVersion.Initialize(FEngineVersion::Current());
+
+	// For builds synced via UGS, we override the changelist of the engine version with the current build version changelist
+	// as this helps to keep the changelists of programmers and artists/designers in-sync when creating and joining sessions
+	//	eg) CL# 1 is a code change, and CL# 2 is a content change:
+	//	 - A programmer syncing CL# 2 would have an engine version with a CL# of 2 (from building their own editor), and a build version CL# of 2 (from UGS).
+	//	 - An artist/designer syncing CL# 2 would have an engine version with a CL# of 1 (from the pre-built editor), but a build version CL# of 2 (from UGS).
+	{
+		// Read the default data (rather than the executable specific data), as the default data 
+		// is updated when syncing, but the executable data is only updated when compiling
+		FBuildVersion BuildVersion;
+		if (FBuildVersion::TryRead(FBuildVersion::GetDefaultFileName(), BuildVersion))
+		{
+			// Only apply the build version if our engine changelist is compatible with the synced build
+			// If this check fails then it likely means that a programmer synced (updating the build version) 
+			// without also compiling their binaries (to update the engine version)
+			if (EngineVersion.Changelist >= (uint32)BuildVersion.CompatibleChangelist)
+			{
+				EngineVersion.Changelist = BuildVersion.Changelist;
+			}
+		}
+	}
 
 	for (const FCustomVersion& EngineCustomVersion : FCustomVersionContainer::GetRegistered().GetAllVersions())
 	{
@@ -109,7 +131,7 @@ bool FConcertSessionVersionInfo::Validate(const FConcertSessionVersionInfo& InOt
 		return false;
 	}
 	
-	if (!CompatibleEngineVersion.Validate(InOther.CompatibleEngineVersion, InValidationMode, OutFailureReason))
+	if (!EngineVersion.Validate(InOther.EngineVersion, InValidationMode, OutFailureReason))
 	{
 		return false;
 	}
diff --git a/Engine/Plugins/Developer/Concert/ConcertMain/Source/Concert/Public/ConcertSettings.h b/Engine/Plugins/Developer/Concert/ConcertMain/Source/Concert/Public/ConcertSettings.h
--- a/Engine/Plugins/Developer/Concert/ConcertMain/Source/Concert/Public/ConcertSettings.h
+++ b/Engine/Plugins/Developer/Concert/ConcertMain/Source/Concert/Public/ConcertSettings.h
@@ -26,14 +26,10 @@ struct FConcertSessionSettings
 {
 	GENERATED_BODY()
 
-	FConcertSessionSettings()
-		: BaseRevision(0)
-	{}
-
 	void Initialize()
 	{
 		ProjectName = FApp::GetProjectName();
-		BaseRevision = FEngineVersion::Current().GetChangelist(); // TODO: This isn't good enough for people using binary builds
+		// TODO: BaseRevision should have a robust way to know which content version a project is on, as we currently check this using the current build version (see EngineVersion in FConcertSessionVersionInfo), which works for UGS but isn't reliable for public binary builds
 	}
 
 	bool ValidateRequirements(const FConcertSessionSettings& Other, FText* OutFailureReason = nullptr) const
@@ -71,7 +67,7 @@ struct FConcertSessionSettings
 	 * Can be specified on the server cmd with `-CONCERTREVISION=`
 	 */
 	UPROPERTY(config, VisibleAnywhere, Category="Session Settings")
-	uint32 BaseRevision;
+	uint32 BaseRevision = 0;
 
 	/**
 	 * Override the default name chosen when archiving this session.
@@ -88,18 +84,13 @@ struct FConcertServerSettings
 {
 	GENERATED_BODY()
 
-	FConcertServerSettings()
-		: bIgnoreSessionSettingsRestriction(false)
-		, SessionTickFrequencySeconds(1)
-	{}
-
 	/** The server will allow client to join potentially incompatible sessions */
 	UPROPERTY(config, EditAnywhere, AdvancedDisplay, Category="Server Settings")
-	bool bIgnoreSessionSettingsRestriction;
+	bool bIgnoreSessionSettingsRestriction = false;
 
 	/** The timespan at which session updates are processed. */
 	UPROPERTY(config, EditAnywhere, DisplayName="Session Tick Frequency", AdvancedDisplay, Category="Server Settings", meta=(ForceUnits=s))
-	int32 SessionTickFrequencySeconds;
+	int32 SessionTickFrequencySeconds = 1;
 };
 
 UCLASS(config=Engine)
diff --git a/Engine/Plugins/Developer/Concert/ConcertMain/Source/Concert/Public/ConcertVersion.h b/Engine/Plugins/Developer/Concert/ConcertMain/Source/Concert/Public/ConcertVersion.h
--- a/Engine/Plugins/Developer/Concert/ConcertMain/Source/Concert/Public/ConcertVersion.h
+++ b/Engine/Plugins/Developer/Concert/ConcertMain/Source/Concert/Public/ConcertVersion.h
@@ -104,15 +104,15 @@ struct FConcertSessionVersionInfo
 	/** Validate this version info against another */
 	CONCERT_API bool Validate(const FConcertSessionVersionInfo& InOther, const EConcertVersionValidationMode InValidationMode, FText* OutFailureReason = nullptr) const;
 
-	/* File version info */
+	/** File version info */
 	UPROPERTY()
 	FConcertFileVersionInfo FileVersion;
 
-	/* Compatible engine version info */
+	/** Engine version info */
 	UPROPERTY()
-	FConcertEngineVersionInfo CompatibleEngineVersion;
+	FConcertEngineVersionInfo EngineVersion;
 
-	/* Custom version info */
+	/** Custom version info */
 	UPROPERTY()
 	TArray<FConcertCustomVersionInfo> CustomVersions;
 };
diff --git a/Engine/Plugins/Developer/Concert/ConcertSync/ConcertSyncCore/Source/ConcertSyncCore/Private/ConcertSyncArchives.cpp b/Engine/Plugins/Developer/Concert/ConcertSync/ConcertSyncCore/Source/ConcertSyncCore/Private/ConcertSyncArchives.cpp
--- a/Engine/Plugins/Developer/Concert/ConcertSync/ConcertSyncCore/Source/ConcertSyncCore/Private/ConcertSyncArchives.cpp
+++ b/Engine/Plugins/Developer/Concert/ConcertSync/ConcertSyncCore/Source/ConcertSyncCore/Private/ConcertSyncArchives.cpp
@@ -176,7 +176,7 @@ FConcertSyncObjectReader::FConcertSyncObjectReader(const FConcertLocalIdentifier
 	{
 		SetUE4Ver(InVersionInfo->FileVersion.FileVersionUE4);
 		SetLicenseeUE4Ver(InVersionInfo->FileVersion.FileVersionLicenseeUE4);
-		SetEngineVer(FEngineVersionBase(InVersionInfo->CompatibleEngineVersion.Major, InVersionInfo->CompatibleEngineVersion.Minor, InVersionInfo->CompatibleEngineVersion.Patch, InVersionInfo->CompatibleEngineVersion.Changelist));
+		SetEngineVer(FEngineVersionBase(InVersionInfo->EngineVersion.Major, InVersionInfo->EngineVersion.Minor, InVersionInfo->EngineVersion.Patch, InVersionInfo->EngineVersion.Changelist));
 
 		FCustomVersionContainer EngineCustomVersions;
 		for (const FConcertCustomVersionInfo& CustomVersion : InVersionInfo->CustomVersions)
diff --git a/Engine/Plugins/Developer/Concert/MultiUserClient/Source/MultiUserClient/Private/Widgets/SConcertBrowser.cpp b/Engine/Plugins/Developer/Concert/MultiUserClient/Source/MultiUserClient/Private/Widgets/SConcertBrowser.cpp
--- a/Engine/Plugins/Developer/Concert/MultiUserClient/Source/MultiUserClient/Private/Widgets/SConcertBrowser.cpp
+++ b/Engine/Plugins/Developer/Concert/MultiUserClient/Source/MultiUserClient/Private/Widgets/SConcertBrowser.cpp
@@ -3146,7 +3146,19 @@ void SConcertSessionBrowser::PopulateSessionInfoGrid(SGridPanel& Grid, const FCo
 	AddDetailRow(Grid, Row++, LOCTEXT("SessionName", "Session Name:"), FText::FromString(SessionInfo.SessionName));
 	AddDetailRow(Grid, Row++, LOCTEXT("Owner", "Owner:"), FText::FromString(SessionInfo.OwnerUserName));
 	AddDetailRow(Grid, Row++, LOCTEXT("Project", "Project:"), FText::FromString(SessionInfo.Settings.ProjectName));
-	AddDetailRow(Grid, Row++, LOCTEXT("BaseVersion", "Base Version:"), FText::AsNumber(SessionInfo.Settings.BaseRevision, &FNumberFormattingOptions::DefaultNoGrouping()));
+	//AddDetailRow(Grid, Row++, LOCTEXT("BaseVersion", "Base Version:"), FText::AsNumber(SessionInfo.Settings.BaseRevision, &FNumberFormattingOptions::DefaultNoGrouping()));
+	if (SessionInfo.VersionInfos.Num() > 0)
+	{
+		const FConcertSessionVersionInfo& VersionInfo = SessionInfo.VersionInfos.Last();
+		AddDetailRow(Grid, Row++, LOCTEXT("EngineVersion", "Engine Version:"),
+			FText::Format(
+				LOCTEXT("EngineVersionFmt", "{0}.{1}.{2}-{3}"),
+				FText::AsNumber(VersionInfo.EngineVersion.Major, &FNumberFormattingOptions::DefaultNoGrouping()),
+				FText::AsNumber(VersionInfo.EngineVersion.Minor, &FNumberFormattingOptions::DefaultNoGrouping()),
+				FText::AsNumber(VersionInfo.EngineVersion.Patch, &FNumberFormattingOptions::DefaultNoGrouping()),
+				FText::AsNumber(VersionInfo.EngineVersion.Changelist, &FNumberFormattingOptions::DefaultNoGrouping())
+			));
+	}
 	AddDetailRow(Grid, Row++, LOCTEXT("ServerEndPointId", "Server Endpoint ID:"), FText::FromString(SessionInfo.ServerEndpointId.ToString()));
 }
 
diff --git a/Engine/Source/Runtime/Core/Private/Modules/BuildVersion.cpp b/Engine/Source/Runtime/Core/Private/Modules/BuildVersion.cpp
--- a/Engine/Source/Runtime/Core/Private/Modules/BuildVersion.cpp
+++ b/Engine/Source/Runtime/Core/Private/Modules/BuildVersion.cpp
@@ -151,6 +151,13 @@ bool FBuildVersion::TryRead(const FString& FileName, FBuildVersion& OutVersion)
 				return false;
 			}
 		}
+		else if (Field == TEXT("BuildVersion"))
+		{
+			if (!FSimpleParse::ParseString(Ptr, OutVersion.BuildVersion))
+			{
+				return false;
+			}
+		}
 		else
 		{
 			return false;
diff --git a/Engine/Source/Runtime/Core/Public/Modules/BuildVersion.h b/Engine/Source/Runtime/Core/Public/Modules/BuildVersion.h
--- a/Engine/Source/Runtime/Core/Public/Modules/BuildVersion.h
+++ b/Engine/Source/Runtime/Core/Public/Modules/BuildVersion.h
@@ -58,6 +58,11 @@ public:
 	 */
 	FString BuildId;
 
+	/**
+	 * The build version string.
+	 */
+	FString BuildVersion;
+
 	/**
 	 * Default constructor. Initializes the structure to empty.
 	 */
