Fix warning when cooking with overrides to Multi-User settings
The cooker was considering editor-only presence avatar as proper cooking targets

diff --git a/Engine/Plugins/Developer/Concert/ConcertMain/Source/Concert/Public/ConcertSettings.h b/Engine/Plugins/Developer/Concert/ConcertMain/Source/Concert/Public/ConcertSettings.h
--- a/Engine/Plugins/Developer/Concert/ConcertMain/Source/Concert/Public/ConcertSettings.h
+++ b/Engine/Plugins/Developer/Concert/ConcertMain/Source/Concert/Public/ConcertSettings.h
@@ -222,6 +222,17 @@ class CONCERT_API UConcertClientConfig : public UObject
 public:
 	UConcertClientConfig();
 
+	/*
+	 * Mark this setting object as editor only.
+	 * This so soft object path reference made by this setting object won't be automatically grabbed by the cooker.
+	 * @see UPackage::Save, FSoftObjectPathThreadContext::GetSerializationOptions, FSoftObjectPath::ImportTextItem
+	 * @todo: cooker should have a better way to filter editor only objects for 'unsollicited' references.
+	 */
+	virtual bool IsEditorOnly() const
+	{
+		return true;
+	}
+
 	/**
 	 * True if this client should be "headless"? (ie, not display any UI).
 	 */
