- Fixed Disaster Recovery not rotating/triming the old recovery sessions because it tried to delete from the wrong directory.
  - Fixed off-by-one error when rotating sessions.
  - Fixed incorrect the base working/saved path in the Concert server when the default paths were used.

diff --git a/Engine/Plugins/Developer/Concert/ConcertMain/Source/Concert/Private/ConcertServer.cpp b/Engine/Plugins/Developer/Concert/ConcertMain/Source/Concert/Private/ConcertServer.cpp
--- a/Engine/Plugins/Developer/Concert/ConcertMain/Source/Concert/Private/ConcertServer.cpp
+++ b/Engine/Plugins/Developer/Concert/ConcertMain/Source/Concert/Private/ConcertServer.cpp
@@ -37,8 +37,8 @@ FString GetArchiveName(const FString& SessionName, const FConcertSessionSettings
 FConcertServerPaths::FConcertServerPaths(const FString& InRole, const FString& InWorkingDir, const FString& InSavedDir)
 	: WorkingDir(InWorkingDir.Len() ? InWorkingDir / InRole : (FPaths::ProjectIntermediateDir() / TEXT("Concert") / InRole))
 	, SavedDir(InSavedDir.Len() ? InSavedDir / InRole : (FPaths::ProjectSavedDir() / TEXT("Concert") / InRole))
-	, BaseWorkingDir(InWorkingDir)
-	, BaseSavedDir(InSavedDir)
+	, BaseWorkingDir(InWorkingDir.Len() ? InWorkingDir : FPaths::ProjectIntermediateDir())
+	, BaseSavedDir(InSavedDir.Len() ? InSavedDir : FPaths::ProjectSavedDir())
 {
 }
 
@@ -351,7 +351,7 @@ void FConcertServer::RecoverSessions()
 			// Keep the most recent sessions
 			TArray<FConcertSessionInfo> ArchivedSessionsToKeep;
 			{
-				const int32 FirstSortedSessionIndexToKeep = SortedSessions.Num() - Settings->NumSessionsToKeep - 1;
+				const int32 FirstSortedSessionIndexToKeep = SortedSessions.Num() - Settings->NumSessionsToKeep;
 				for (int32 SortedSessionIndex = FirstSortedSessionIndexToKeep; SortedSessionIndex < SortedSessions.Num(); ++SortedSessionIndex)
 				{
 					ArchivedSessionsToKeep.Add(ArchivedSessionInfos[SortedSessions[SortedSessionIndex].Key]);
@@ -362,7 +362,7 @@ void FConcertServer::RecoverSessions()
 			// Remove the oldest sessions
 			for (const FSavedSessionInfo& SortedSession : SortedSessions)
 			{
-				ConcertUtil::DeleteDirectoryTree(*Paths->GetSessionWorkingDir(ArchivedSessionInfos[SortedSession.Key].SessionId), *Paths->GetBaseWorkingDir());
+				ConcertUtil::DeleteDirectoryTree(*Paths->GetSessionSavedDir(ArchivedSessionInfos[SortedSession.Key].SessionId), *Paths->GetBaseSavedDir());	
 			}
 
 			// Update the list of sessions to restore
