- Fixed a disaster recovery bug preventing the Editor from recovering a session because another instance of the Editor on another project already locked all the sessions.

Problem:

On windows, the CrashReportClientEditor (hosting disaster recovery service) is started in the static initialization, before the engine is initialized, not allowing lot of command line configuration. The Editor project browser would start a first CrashReportClientEditor instance, which would load and lock all the available sessions (unless another CrashReportClientEditor was running). When the user selected a project, a new Editor and CrashReportClientEditor were launched before the first one was closed. The second instance could not access the existing sessions because they were still locked by the first instance.

Solution:

Because CrashReportClientEditor is launch before the engine is initialized, we don't have any context at the launch time. The best the was to delay the moment when the server reloads the existing sessions and enable each clients to store their sessions in different folders (repositories) mounted on demand by the server.

Implementation details:
  - Implemented new RPC API to allow the client to list/create/load/drop specific repositories containing its own sessions on demand.
  - Updated the Concert server to manage multiples directories where session can be stored/found (session repositories) rather than just one.
  - Added a settings to allow the user to specify where the disaster recovery sessions should be stored on the disk. Now default in the current project folder.
  - Added a settings to prevent the Concert server from scanning the sessions in the default location.
  - Updated disaster recovery to start without any session repository and let the client decide if a new one needs to be created or an existing one be mounted to restore a previous session.
  - Changed the code to let disaster recovery client manage its session history rather than letting the server rotate the old session. Defaulted the history to 0, user has no flow to visualize and pick from the history.

diff --git a/Engine/Plugins/Developer/Concert/ConcertMain/Source/Concert/Private/ConcertClient.cpp b/Engine/Plugins/Developer/Concert/ConcertMain/Source/Concert/Private/ConcertClient.cpp
--- a/Engine/Plugins/Developer/Concert/ConcertMain/Source/Concert/Private/ConcertClient.cpp
+++ b/Engine/Plugins/Developer/Concert/ConcertMain/Source/Concert/Private/ConcertClient.cpp
@@ -972,6 +972,29 @@ TSharedPtr<IConcertClientSession> FConcertClient::GetCurrentSession() const
 	return ClientSession;
 }
 
+TFuture<FConcertAdmin_MountSessionRepositoryResponse> FConcertClient::MountSessionRepository(const FGuid& ServerAdminEndpointId, const FString& RepositoryRootDir, const FGuid& RepositoryId, bool bCreateIfNotExist, bool bAsDefault) const
+{
+	FConcertAdmin_MountSessionRepositoryRequest MountRepositoryRequest;
+	MountRepositoryRequest.RepositoryId = RepositoryId;
+	MountRepositoryRequest.RepositoryRootDir = RepositoryRootDir;
+	MountRepositoryRequest.bCreateIfNotExist = bCreateIfNotExist;
+	MountRepositoryRequest.bAsServerDefault = bAsDefault;
+	return ClientAdminEndpoint->SendRequest<FConcertAdmin_MountSessionRepositoryRequest, FConcertAdmin_MountSessionRepositoryResponse>(MountRepositoryRequest, ServerAdminEndpointId);
+}
+
+TFuture<FConcertAdmin_GetSessionRepositoriesResponse> FConcertClient::GetSessionRepositories(const FGuid& ServerAdminEndpointId) const
+{
+	FConcertAdmin_GetSessionRepositoriesRequest GetRepositoryRequest;
+	return ClientAdminEndpoint->SendRequest<FConcertAdmin_GetSessionRepositoriesRequest, FConcertAdmin_GetSessionRepositoriesResponse>(GetRepositoryRequest, ServerAdminEndpointId);
+}
+
+TFuture<FConcertAdmin_DropSessionRepositoriesResponse> FConcertClient::DropSessionRepositories(const FGuid& ServerAdminEndpointId, const TArray<FGuid>& RepositoryIds) const
+{
+	FConcertAdmin_DropSessionRepositoriesRequest DropRepositoryRequest;
+	DropRepositoryRequest.RepositoryIds = RepositoryIds;
+	return ClientAdminEndpoint->SendRequest<FConcertAdmin_DropSessionRepositoriesRequest, FConcertAdmin_DropSessionRepositoriesResponse>(DropRepositoryRequest, ServerAdminEndpointId);
+}
+
 TFuture<FConcertAdmin_GetAllSessionsResponse> FConcertClient::GetServerSessions(const FGuid& ServerAdminEndpointId) const
 {
 	FConcertAdmin_GetAllSessionsRequest GetSessionsRequest = FConcertAdmin_GetAllSessionsRequest();
diff --git a/Engine/Plugins/Developer/Concert/ConcertMain/Source/Concert/Private/ConcertClient.h b/Engine/Plugins/Developer/Concert/ConcertMain/Source/Concert/Private/ConcertClient.h
--- a/Engine/Plugins/Developer/Concert/ConcertMain/Source/Concert/Private/ConcertClient.h
+++ b/Engine/Plugins/Developer/Concert/ConcertMain/Source/Concert/Private/ConcertClient.h
@@ -88,6 +88,9 @@ public:
 	virtual bool IsOwnerOf(const FConcertSessionInfo& InSessionInfo) const override;
 	virtual TSharedPtr<IConcertClientSession> GetCurrentSession() const override;
 
+	virtual TFuture<FConcertAdmin_MountSessionRepositoryResponse> MountSessionRepository(const FGuid& ServerAdminEndpointId, const FString& RepositoryRootDir, const FGuid& RepositoryId, bool bCreateIfNotExist, bool bAsDefault) const override;
+	virtual TFuture<FConcertAdmin_GetSessionRepositoriesResponse> GetSessionRepositories(const FGuid& ServerAdminEndpointId) const override;
+	virtual TFuture<FConcertAdmin_DropSessionRepositoriesResponse> DropSessionRepositories(const FGuid& ServerAdminEndpointId, const TArray<FGuid>& RepositoryIds) const override;
 	virtual TFuture<FConcertAdmin_GetAllSessionsResponse> GetServerSessions(const FGuid& ServerAdminEndpointId) const override;
 	virtual TFuture<FConcertAdmin_GetSessionsResponse> GetLiveSessions(const FGuid& ServerAdminEndpointId) const override;
 	virtual TFuture<FConcertAdmin_GetSessionsResponse> GetArchivedSessions(const FGuid& ServerAdminEndpointId) const override;
diff --git a/Engine/Plugins/Developer/Concert/ConcertMain/Source/Concert/Private/ConcertServer.cpp b/Engine/Plugins/Developer/Concert/ConcertMain/Source/Concert/Private/ConcertServer.cpp
--- a/Engine/Plugins/Developer/Concert/ConcertMain/Source/Concert/Private/ConcertServer.cpp
+++ b/Engine/Plugins/Developer/Concert/ConcertMain/Source/Concert/Private/ConcertServer.cpp
@@ -6,6 +6,7 @@
 #include "ConcertLogger.h"
 #include "ConcertSettings.h"
 #include "ConcertServerSession.h"
+#include "ConcertServerSessionRepositories.h"
 #include "ConcertLogGlobal.h"
 #include "IConcertServerEventSink.h"
 
@@ -14,7 +15,6 @@
 #include "Misc/Paths.h"
 
 #include "Runtime/Launch/Resources/Version.h"
-#include "ConcertServerInstanceInfo.h"
 #include "StructSerializer.h"
 #include "StructDeserializer.h"
 #include "Backends/JsonStructSerializerBackend.h"
@@ -44,17 +44,17 @@ FString GetArchiveName(const FString& SessionName, const FConcertSessionSettings
 	}
 }
 
-FString GetServerInstanceInfoPathname(const FString& Role)
+FString GetSessionRepositoryDatabasePathname(const FString& Role)
 {
-	return FPaths::ProjectSavedDir() / Role / TEXT("InstanceInfo.json");
+	return FPaths::ProjectSavedDir() / Role / TEXT("Repositories.json");
 }
 
-bool SaveServerInstanceInfo(const FString& Role, const FConcertServerInstanceInfo& InInstanceInfo)
+bool SaveSessionRepositoryDatabase(const FString& Role, const FConcertServerSessionRepositoryDatabase& RepositoryDb)
 {
-	if (TUniquePtr<FArchive> FileWriter = TUniquePtr<FArchive>(IFileManager::Get().CreateFileWriter(*GetServerInstanceInfoPathname(Role))))
+	if (TUniquePtr<FArchive> FileWriter = TUniquePtr<FArchive>(IFileManager::Get().CreateFileWriter(*GetSessionRepositoryDatabasePathname(Role))))
 	{
 		FJsonStructSerializerBackend Backend(*FileWriter, EStructSerializerBackendFlags::Default);
-		FStructSerializer::Serialize(InInstanceInfo, Backend);
+		FStructSerializer::Serialize(RepositoryDb, Backend);
 
 		FileWriter->Close();
 		return !FileWriter->IsError();
@@ -63,12 +63,12 @@ bool SaveServerInstanceInfo(const FString& Role, const FConcertServerInstanceInf
 	return false;
 }
 
-bool LoadServerInstanceInfo(const FString& Role, FConcertServerInstanceInfo& OutIntanceInfo)
+bool LoadSessionRepositoryDatabase(const FString& Role, FConcertServerSessionRepositoryDatabase& RepositoryDb)
 {
-	if (TUniquePtr<FArchive> FileReader = TUniquePtr<FArchive>(IFileManager::Get().CreateFileReader(*GetServerInstanceInfoPathname(Role))))
+	if (TUniquePtr<FArchive> FileReader = TUniquePtr<FArchive>(IFileManager::Get().CreateFileReader(*GetSessionRepositoryDatabasePathname(Role))))
 	{
 		FJsonStructDeserializerBackend Backend(*FileReader);
-		FStructDeserializer::Deserialize(OutIntanceInfo, Backend);
+		FStructDeserializer::Deserialize(RepositoryDb, Backend);
 
 		FileReader->Close();
 		return !FileReader->IsError();
@@ -77,38 +77,6 @@ bool LoadServerInstanceInfo(const FString& Role, FConcertServerInstanceInfo& Out
 	return false;
 }
 
-int32 RemoveInactiveServerInstances(FConcertServerInstanceInfo& InOutIntanceInfo)
-{
-	InOutIntanceInfo.Instances.RemoveAll([](const FConcertServerInstance& Instance)
-	{
-		if (FPlatformProcess::IsApplicationRunning(Instance.ProcessId))
-		{
-			// If the PID was reused by another application, remove it from the list.
-			return FPaths::GetPathLeaf(FPlatformProcess::GetApplicationName(Instance.ProcessId)) != FPaths::GetPathLeaf(FPlatformProcess::GetApplicationName(FPlatformProcess::GetCurrentProcessId()));
-		}
-		return true; // Not running, remove it.
-	});
-
-	return InOutIntanceInfo.Instances.Num(); // How many remains?
-}
-
-bool SharesDirectoriesWithOtherInstances(const FConcertServerPaths& Paths, FConcertServerInstanceInfo& InstanceInfo)
-{
-	return Algo::AnyOf(InstanceInfo.Instances, [&Paths](const FConcertServerInstance& Instance)
-	{
-		return Instance.ArchiveDirectory == Paths.GetSavedDir() || Instance.WorkingDirectory == Paths.GetWorkingDir();
-	});
-}
-
-}
-
-
-FConcertServerPaths::FConcertServerPaths(const FString& InRole, const FString& InWorkingDir, const FString& InSavedDir)
-	: WorkingDir(InWorkingDir.Len() ? InWorkingDir / InRole : (FPaths::ProjectIntermediateDir() / TEXT("Concert") / InRole))
-	, SavedDir(InSavedDir.Len() ? InSavedDir / InRole : (FPaths::ProjectSavedDir() / TEXT("Concert") / InRole))
-	, BaseWorkingDir(InWorkingDir.Len() ? InWorkingDir : FPaths::ProjectIntermediateDir())
-	, BaseSavedDir(InSavedDir.Len() ? InSavedDir : FPaths::ProjectSavedDir())
-{
 }
 
 FConcertServer::FConcertServer(const FString& InRole, const FConcertSessionFilter& InAutoArchiveSessionFilter, IConcertServerEventSink* InEventSink, const TSharedPtr<IConcertEndpointProvider>& InEndpointProvider)
@@ -133,12 +101,12 @@ const FString& FConcertServer::GetRole() const
 
 void FConcertServer::Configure(const UConcertServerConfig* InSettings)
 {
+	check(!Settings); // Server do not support reconfiguration.
+
 	ServerInfo.Initialize();
 	check(InSettings != nullptr);
 	Settings = TStrongObjectPtr<const UConcertServerConfig>(InSettings);
 
-	Paths = MakeUnique<FConcertServerPaths>(GetRole(), InSettings->WorkingDir, InSettings->ArchiveDir);
-
 	if (!InSettings->ServerName.IsEmpty())
 	{
 		ServerInfo.ServerName = InSettings->ServerName;
@@ -148,6 +116,19 @@ void FConcertServer::Configure(const UConcertServerConfig* InSettings)
 	{
 		ServerInfo.ServerFlags |= EConcertServerFlags::IgnoreSessionRequirement;
 	}
+
+	SessionRepositoryRootDir = FPaths::ProjectSavedDir() / Role / TEXT("Sessions"); // Server default session repository root dir.
+	if (!Settings->SessionRepositoryRootDir.IsEmpty())
+	{
+		if (IFileManager::Get().DirectoryExists(*Settings->SessionRepositoryRootDir) || IFileManager::Get().MakeDirectory(*Settings->SessionRepositoryRootDir, /*Tree*/true))
+		{
+			SessionRepositoryRootDir = Settings->SessionRepositoryRootDir; // Overwrite the default.
+		}
+		else
+		{
+			UE_LOG(LogConcert, Warning, TEXT("Invalid session repository root directory. Falling back on %s default."), *SessionRepositoryRootDir);
+		}
+	}
 }
 
 bool FConcertServer::IsConfigured() const
@@ -174,7 +155,6 @@ bool FConcertServer::IsStarted() const
 void FConcertServer::Startup()
 {
 	check(IsConfigured());
-	check (Paths.IsValid());
 	if (!ServerAdminEndpoint.IsValid() && EndpointProvider.IsValid())
 	{
 		// Create the server administration endpoint
@@ -197,45 +177,65 @@ void FConcertServer::Startup()
 		ServerAdminEndpoint->RegisterRequestHandler<FConcertAdmin_GetArchivedSessionsRequest, FConcertAdmin_GetSessionsResponse>(this, &FConcertServer::HandleGetArchivedSessionsRequest);
 		ServerAdminEndpoint->RegisterRequestHandler<FConcertAdmin_GetSessionClientsRequest, FConcertAdmin_GetSessionClientsResponse>(this, &FConcertServer::HandleGetSessionClientsRequest);
 		ServerAdminEndpoint->RegisterRequestHandler<FConcertAdmin_GetSessionActivitiesRequest, FConcertAdmin_GetSessionActivitiesResponse>(this, &FConcertServer::HandleGetSessionActivitiesRequest);
+		ServerAdminEndpoint->RegisterRequestHandler<FConcertAdmin_MountSessionRepositoryRequest, FConcertAdmin_MountSessionRepositoryResponse>(this, &FConcertServer::HandleMountSessionRepositoryRequest);
+		ServerAdminEndpoint->RegisterRequestHandler<FConcertAdmin_GetSessionRepositoriesRequest, FConcertAdmin_GetSessionRepositoriesResponse>(this, &FConcertServer::HandleGetSessionRepositoriesRequest);
+		ServerAdminEndpoint->RegisterRequestHandler<FConcertAdmin_DropSessionRepositoriesRequest, FConcertAdmin_DropSessionRepositoriesResponse>(this, &FConcertServer::HandleDropSessionRepositoriesRequest);
 
-		// Concurrent server instances may fight to get ownership of the info file.
-		FSystemWideCriticalSection ScopedSystemWideMutex(ConcertServerUtil::GetServerSystemMutexName());
+		// Perform maintenance tasks on the session repositories database.
+		{
+			FSystemWideCriticalSection ScopedSystemWideMutex(ConcertServerUtil::GetServerSystemMutexName());
 
-		// Load the file containing the instance info.
-		FConcertServerInstanceInfo InstanceInfo;
-		ConcertServerUtil::LoadServerInstanceInfo(Role, InstanceInfo);
+			// Load the file containing the repositories.
+			FConcertServerSessionRepositoryDatabase SessionRepositoryDb;
+			ConcertServerUtil::LoadSessionRepositoryDatabase(Role, SessionRepositoryDb);
 
-		// Check if the server can scan/load/rotate the existing session files in the configured directories. (The existing session files are not sharable, they can only managed by one instance)
-		int32 ActiveInstanceCount = ConcertServerUtil::RemoveInactiveServerInstances(InstanceInfo);
-		bExclusiveAccessToExistingSessionsAtStartup = ActiveInstanceCount == 0 || !ConcertServerUtil::SharesDirectoriesWithOtherInstances(*Paths, InstanceInfo);
-		if (bExclusiveAccessToExistingSessionsAtStartup)
-		{
-			if (Settings->bCleanWorkingDir)
-			{
-				ConcertUtil::DeleteDirectoryTree(*Paths->GetWorkingDir(), *Paths->GetBaseWorkingDir());
-			}
-			else
+			// Walk the root directory containing the server managed repositories and find those that aren't mapped anymore.
+			TArray<FString> ExpiredDirectories;
+			IFileManager::Get().IterateDirectory(*GetSessionRepositoriesRootDir(), [this, &SessionRepositoryDb, &ExpiredDirectories](const TCHAR* Pathname, bool bIsDirectory)
 			{
-				if (Settings->bAutoArchiveOnReboot)
+				if (bIsDirectory)
 				{
-					// Migrate live sessions files (session is not restored yet) to its archive form and directory.
-					ArchiveOfflineSessions();
+					// Check if the repository is still mapped.
+					FString RootReposDir = GetSessionRepositoriesRootDir();
+					if (!SessionRepositoryDb.Repositories.ContainsByPredicate([&RootReposDir, Pathname](const FConcertServerSessionRepository& Repository) { return  RootReposDir / Repository.RepositoryId.ToString() == Pathname; }))
+					{
+						ExpiredDirectories.Emplace(Pathname); // The visited directory was not found in the list of repositories.
+					}
 				}
+				return true;
+			});
 
-				// Build the list of archive/live sessions and rotate the list of archive to prevent having too many of them.
-				RecoverSessions();
+			// Delete the directories that are not mapped anymore.
+			for (const FString& Dir : ExpiredDirectories)
+			{
+				FGuid Dummy;
+				if (FGuid::Parse(FPaths::GetPathLeaf(Dir), Dummy)) // Ensure the directory name is a GUID as the repositories base dir is the repository ID.
+				{
+					ConcertUtil::DeleteDirectoryTree(*Dir);
+				}
 			}
 		}
-		// else -> This instance will be able to create/rename/delete new sessions but will not see/expose the currently existing ones to avoid file conflict/corruption.
 
-		// Add this instance to the instance info file.
-		FConcertServerInstance& ThisInstance = InstanceInfo.Instances.AddDefaulted_GetRef();
-		ThisInstance.ServerName = Settings->ServerName;
-		ThisInstance.WorkingDirectory = Paths->GetWorkingDir();
-		ThisInstance.ArchiveDirectory = Paths->GetSavedDir();
-		ThisInstance.ProcessId = FPlatformProcess::GetCurrentProcessId();
-		ThisInstance.bEnclaved = !bExclusiveAccessToExistingSessionsAtStartup; // Enclaved means this instance did not scan and load existing sessions to avoid conflict with another running instance(s). It was restricted to work with its own sessions only.
-		ConcertServerUtil::SaveServerInstanceInfo(Role, InstanceInfo);
+		// If the server was configured to use a custom working/archive dir, create a corresponding repository and try to mount it.
+		if (!Settings->WorkingDir.IsEmpty() || !Settings->ArchiveDir.IsEmpty())
+		{
+			FConcertServerSessionRepository Repository(Role, FGuid::NewGuid(), Settings->WorkingDir, Settings->ArchiveDir);
+			if (MountSessionRepository(Repository, /*bCreateIfNotExist*/true, Settings->bCleanWorkingDir, /*bCleanupExpiredSession*/true, /*bSearchByPath*/true) == EConcertSessionRepositoryMountResponseCode::Mounted)
+			{
+				check(Repository.bMounted && Repository.ProcessId == FPlatformProcess::GetCurrentProcessId());
+				DefaultSessionRepository.Emplace(Repository);
+			}
+		}
+		// If the server was configured to mount a default server managed repository.
+		else if (Settings->bMountDefaultSessionRepository)
+		{
+			FConcertServerSessionRepository Repository(GetSessionRepositoriesRootDir(), FGuid()); // Invalid GUID is used for the default server repository.
+			if (MountSessionRepository(Repository, /*bCreateIfNotExist*/true, Settings->bCleanWorkingDir, /*bCleanupExpiredSession*/true, /*bSearchByPath*/false) == EConcertSessionRepositoryMountResponseCode::Mounted)
+			{
+				check(Repository.bMounted && Repository.ProcessId == FPlatformProcess::GetCurrentProcessId());
+				DefaultSessionRepository.Emplace(Repository);
+			}
+		}
 	}
 }
 
@@ -296,17 +296,26 @@ void FConcertServer::Shutdown()
 	FSystemWideCriticalSection ScopedSystemWideMutex(ConcertServerUtil::GetServerSystemMutexName());
 
 	// Load the file containing the instance info.
-	FConcertServerInstanceInfo InstanceInfo;
-	ConcertServerUtil::LoadServerInstanceInfo(Role, InstanceInfo);
+	FConcertServerSessionRepositoryDatabase SessionRepositoryDb;
+	ConcertServerUtil::LoadSessionRepositoryDatabase(Role, SessionRepositoryDb);
+	bool bSaveRepositoryDatabase = false;
 
-	// Remove this instance form the list of instances.
+	// Unmount all repositories mounted by this instance.
 	int32 ProcessId = FPlatformProcess::GetCurrentProcessId();
-	InstanceInfo.Instances.RemoveAll([ProcessId](const FConcertServerInstance& Instance)
+	for (FConcertServerSessionRepository& Repository : SessionRepositoryDb.Repositories)
 	{
-		return Instance.ProcessId == ProcessId;
-	});
+		if (Repository.bMounted && Repository.ProcessId == ProcessId)
+		{
+			Repository.bMounted = false;
+			Repository.ProcessId = 0;
+			bSaveRepositoryDatabase = true;
+		}
+	}
 
-	ConcertServerUtil::SaveServerInstanceInfo(Role, InstanceInfo);
+	if (bSaveRepositoryDatabase)
+	{
+		ConcertServerUtil::SaveSessionRepositoryDatabase(Role, SessionRepositoryDb);
+	}
 }
 
 FGuid FConcertServer::GetLiveSessionIdByName(const FString& InName) const
@@ -374,7 +383,14 @@ TSharedPtr<IConcertServerSession> FConcertServer::CreateSession(const FConcertSe
 		return nullptr;
 	}
 
-	return CreateLiveSession(SessionInfo);
+	if (!DefaultSessionRepository)
+	{
+		OutFailureReason = FText::Format(LOCTEXT("Error_CreateSession_NoRepository", "Session '{0}' could not be created, no session repository was mounted to store it"), FText::AsCultureInvariant(SessionInfo.SessionName));
+		UE_LOG(LogConcert, Error, TEXT("An attempt to create a session with name '%s' was made, but the server did not have any repository mounted to store it!"), *SessionInfo.SessionName);
+		return nullptr;
+	}
+
+	return CreateLiveSession(SessionInfo, DefaultSessionRepository.GetValue());
 }
 
 TSharedPtr<IConcertServerSession> FConcertServer::RestoreSession(const FGuid& SessionId, const FConcertSessionInfo& SessionInfo, const FConcertSessionFilter& SessionFilter, FText& OutFailureReason)
@@ -410,28 +426,26 @@ TSharedPtr<IConcertServerSession> FConcertServer::RestoreSession(const FGuid& Se
 	return RestoreArchivedSession(SessionId, SessionInfo, SessionFilter, OutFailureReason);
 }
 
-void FConcertServer::RecoverSessions()
+void FConcertServer::RecoverSessions(const FConcertServerSessionRepository& InRepository, bool bCleanupExpiredSessions)
 {
-	check(LiveSessions.Num() == 0 && ArchivedSessions.Num() == 0);
-
 	// Find any existing live sessions to automatically restore when recovering from an improper server shutdown
 	TArray<FConcertSessionInfo> LiveSessionInfos;
-	EventSink->GetSessionsFromPath(*this, Paths->GetWorkingDir(), LiveSessionInfos);
+	EventSink->GetSessionsFromPath(*this, InRepository.WorkingDir, LiveSessionInfos);
 
 	// Restore any existing live sessions
 	for (FConcertSessionInfo& LiveSessionInfo : LiveSessionInfos)
 	{
 		// Update the session info with new server info
 		LiveSessionInfo.ServerInstanceId = ServerInfo.InstanceInfo.InstanceId;
-		if (!LiveSessions.Contains(LiveSessionInfo.SessionId) && !GetLiveSessionIdByName(LiveSessionInfo.SessionName).IsValid() && CreateLiveSession(LiveSessionInfo))
+		if (!LiveSessions.Contains(LiveSessionInfo.SessionId) && !GetLiveSessionIdByName(LiveSessionInfo.SessionName).IsValid() && CreateLiveSession(LiveSessionInfo, InRepository))
 		{
 			UE_LOG(LogConcert, Display, TEXT("Live session '%s' (%s) was recovered."), *LiveSessionInfo.SessionName, *LiveSessionInfo.SessionId.ToString());
 		}
 	}
 
-	if (Settings->NumSessionsToKeep == 0)
+	if (bCleanupExpiredSessions && Settings->NumSessionsToKeep == 0)
 	{
-		ConcertUtil::DeleteDirectoryTree(*Paths->GetSavedDir(), *Paths->GetBaseSavedDir());
+		ConcertUtil::DeleteDirectoryTree(*InRepository.SavedDir);
 	}
 	else
 	{
@@ -440,11 +454,11 @@ void FConcertServer::RecoverSessions()
 		TArray<FDateTime> ArchivedSessionCreationTimes;
 
 		// In theory, archives are immutable, but the server will end up touching the files and change the 'modification time'. Ensure to look at 'creation time'.
-		EventSink->GetSessionsFromPath(*this, Paths->GetSavedDir(), ArchivedSessionInfos, &ArchivedSessionCreationTimes);
+		EventSink->GetSessionsFromPath(*this, InRepository.SavedDir, ArchivedSessionInfos, &ArchivedSessionCreationTimes);
 		check(ArchivedSessionInfos.Num() == ArchivedSessionCreationTimes.Num());
 
-		// Trim the oldest archived sessions
-		if (Settings->NumSessionsToKeep > 0 && ArchivedSessionInfos.Num() > Settings->NumSessionsToKeep)
+		// Trim the oldest archived sessions.
+		if (bCleanupExpiredSessions && Settings->NumSessionsToKeep > 0 && ArchivedSessionInfos.Num() > Settings->NumSessionsToKeep)
 		{
 			typedef TTuple<int32, FDateTime> FSavedSessionInfo;
 
@@ -473,7 +487,7 @@ void FConcertServer::RecoverSessions()
 			// Remove the oldest sessions
 			for (const FSavedSessionInfo& SortedSession : SortedSessions)
 			{
-				ConcertUtil::DeleteDirectoryTree(*Paths->GetSessionSavedDir(ArchivedSessionInfos[SortedSession.Key].SessionId), *Paths->GetBaseSavedDir());
+				ConcertUtil::DeleteDirectoryTree(*InRepository.GetSessionSavedDir(ArchivedSessionInfos[SortedSession.Key].SessionId));
 			}
 
 			// Update the list of sessions to restore
@@ -494,11 +508,11 @@ void FConcertServer::RecoverSessions()
 	}
 }
 
-void FConcertServer::ArchiveOfflineSessions()
+void FConcertServer::ArchiveOfflineSessions(const FConcertServerSessionRepository& InRepository)
 {
 	// Find existing live session files to automatically archive them when recovering from an improper server shutdown.
 	TArray<FConcertSessionInfo> LiveSessionInfos;
-	EventSink->GetSessionsFromPath(*this, Paths->GetWorkingDir(), LiveSessionInfos);
+	EventSink->GetSessionsFromPath(*this, InRepository.WorkingDir, LiveSessionInfos);
 
 	// Migrate the live sessions files into their archived form.
 	for (FConcertSessionInfo& LiveSessionInfo : LiveSessionInfos)
@@ -508,10 +522,10 @@ void FConcertServer::ArchiveOfflineSessions()
 		ArchivedSessionInfo.SessionId = FGuid::NewGuid();
 		ArchivedSessionInfo.SessionName = ConcertServerUtil::GetArchiveName(LiveSessionInfo.SessionName, LiveSessionInfo.Settings);
 
-		if (EventSink->ArchiveSession(*this, Paths->GetSessionWorkingDir(LiveSessionInfo.SessionId), Paths->GetSessionSavedDir(ArchivedSessionInfo.SessionId), ArchivedSessionInfo, AutoArchiveSessionFilter))
+		if (EventSink->ArchiveSession(*this, InRepository.GetSessionWorkingDir(LiveSessionInfo.SessionId), InRepository.GetSessionSavedDir(ArchivedSessionInfo.SessionId), ArchivedSessionInfo, AutoArchiveSessionFilter))
 		{
-			UE_LOG(LogConcert, Display, TEXT("Deleting %s"), *Paths->GetSessionWorkingDir(LiveSessionInfo.SessionId));
-			ConcertUtil::DeleteDirectoryTree(*Paths->GetSessionWorkingDir(LiveSessionInfo.SessionId), *Paths->GetBaseWorkingDir());
+			UE_LOG(LogConcert, Display, TEXT("Deleting %s"), *InRepository.GetSessionWorkingDir(LiveSessionInfo.SessionId));
+			ConcertUtil::DeleteDirectoryTree(*InRepository.GetSessionWorkingDir(LiveSessionInfo.SessionId));
 			UE_LOG(LogConcert, Display, TEXT("Live session '%s' (%s) was archived on reboot."), *LiveSessionInfo.SessionName, *LiveSessionInfo.SessionId.ToString());
 		}
 	}
@@ -610,6 +624,157 @@ TSharedPtr<IConcertServerSession> FConcertServer::GetSession(const FGuid& Sessio
 	return LiveSessions.FindRef(SessionId);
 }
 
+const FString& FConcertServer::GetSessionRepositoriesRootDir() const
+{
+	return SessionRepositoryRootDir;
+}
+
+const FConcertServerSessionRepository& FConcertServer::GetSessionRepository(const FGuid& SessionId) const
+{
+	const FConcertServerSessionRepository* SessionRepository = MountedSessionRepositories.FindByPredicate([SessionId](const FConcertServerSessionRepository& MountedRepository)
+	{
+		return IFileManager::Get().DirectoryExists(*MountedRepository.GetSessionWorkingDir(SessionId)) || IFileManager::Get().DirectoryExists(*MountedRepository.GetSessionSavedDir(SessionId));
+	});
+
+	check(SessionRepository); // If the session is in memory, its repository must be mounted.
+	return *SessionRepository;
+}
+
+FString FConcertServer::GetSessionSavedDir(const FGuid& SessionId) const
+{
+	return GetSessionRepository(SessionId).GetSessionSavedDir(SessionId);
+}
+
+FString FConcertServer::GetSessionWorkingDir(const FGuid& SessionId) const
+{
+	return GetSessionRepository(SessionId).GetSessionWorkingDir(SessionId);
+}
+
+EConcertSessionRepositoryMountResponseCode FConcertServer::MountSessionRepository(FConcertServerSessionRepository& Repository, bool bCreateIfNotExist, bool bCleanWorkingDir, bool bCleanExpiredSessions, bool bSearchByPaths)
+{
+	// Exclusive access scope to the session repository db.
+	{
+		// Load the file containing the instance/repository info.
+		FSystemWideCriticalSection ScopedSystemWideMutex(ConcertServerUtil::GetServerSystemMutexName());
+		FConcertServerSessionRepositoryDatabase SessionRepositoryDb;
+		ConcertServerUtil::LoadSessionRepositoryDatabase(Role, SessionRepositoryDb);
+
+		check(!Repository.bMounted && Repository.ProcessId == 0); // Should not be mounted.
+
+		// Check if the repository can be found in the database.
+		if (FConcertServerSessionRepository* ExistingRepository = SessionRepositoryDb.Repositories.FindByPredicate(
+			[&Repository, bSearchByPaths](const FConcertServerSessionRepository& Candidate){ return bSearchByPaths ? Candidate.WorkingDir == Repository.WorkingDir && Candidate.SavedDir == Repository.SavedDir : Candidate.RepositoryId == Repository.RepositoryId; }))
+		{
+			if (!ExistingRepository->bMounted || !FPlatformProcess::IsApplicationRunning(ExistingRepository->ProcessId))
+			{
+				check(Repository.RepositoryRootDir == ExistingRepository->RepositoryRootDir) // The client changed the root dir?
+				ExistingRepository->bMounted = true;
+				ExistingRepository->ProcessId = FPlatformProcess::GetCurrentProcessId();
+				Repository = *ExistingRepository;
+				MountedSessionRepositories.Add(Repository);
+				ConcertServerUtil::SaveSessionRepositoryDatabase(Role, SessionRepositoryDb);
+			}
+			else
+			{
+				return EConcertSessionRepositoryMountResponseCode::AlreadyMounted; // Already mounted by another process, cannot mount it.
+			}
+		}
+		else if (bCreateIfNotExist)
+		{
+			Repository.bMounted = true;
+			Repository.ProcessId = FPlatformProcess::GetCurrentProcessId();
+			MountedSessionRepositories.Add(Repository);
+			SessionRepositoryDb.Repositories.Add(Repository);
+			ConcertServerUtil::SaveSessionRepositoryDatabase(Role, SessionRepositoryDb);
+		}
+		else
+		{
+			return EConcertSessionRepositoryMountResponseCode::NotFound;
+		}
+	}
+
+	// Process the sessions in the repository.
+	if (bCleanWorkingDir)
+	{
+		ConcertUtil::DeleteDirectoryTree(*Repository.WorkingDir);
+	}
+	else if (Settings->bAutoArchiveOnReboot) // Honor the auto-archive settings when mounting a new repository.
+	{
+		// Migrate live sessions files (session is not restored yet) to its archive form and directory.
+		ArchiveOfflineSessions(Repository);
+	}
+
+	// Reload the archived/live sessions and possibly rotate the list of archives to prevent having too many of them.
+	RecoverSessions(Repository, bCleanExpiredSessions);
+
+	return EConcertSessionRepositoryMountResponseCode::Mounted;
+}
+
+bool FConcertServer::UnmountSessionRepository(FConcertServerSessionRepository& Repository, bool bDropped)
+{
+	check(Repository.bMounted && Repository.ProcessId == FPlatformProcess::GetCurrentProcessId());
+	check(MountedSessionRepositories.FindByPredicate([Repository](const FConcertServerSessionRepository& Candidate) { return Candidate.RepositoryId == Repository.RepositoryId; })); // Should be found in the list of mounted repositories.
+
+	// Unload the live sessions hosted in that repository.
+	TArray<FGuid> LiveSessionIds;
+	LiveSessions.GetKeys(LiveSessionIds);
+	for (const FGuid& LiveSessionId : LiveSessionIds)
+	{
+		const FConcertServerSessionRepository& SessionRepository = GetSessionRepository(LiveSessionId);
+		if (SessionRepository.RepositoryId == Repository.RepositoryId)
+		{
+			DestroyLiveSession(LiveSessionId, /*bDeleteSessionData*/false);
+		}
+	}
+
+	// Unload the archived sessions hosted in that repository.
+	TArray<FGuid> ArchivedSessionIds;
+	ArchivedSessions.GetKeys(ArchivedSessionIds);
+	for (const FGuid& ArchivedSessionId : ArchivedSessionIds)
+	{
+		const FConcertServerSessionRepository& SessionRepository = GetSessionRepository(ArchivedSessionId);
+		if (SessionRepository.RepositoryId == Repository.RepositoryId)
+		{
+			DestroyArchivedSession(ArchivedSessionId, /*bDeleteSessionData*/false);
+		}
+	}
+
+	if (DefaultSessionRepository && DefaultSessionRepository->RepositoryId == Repository.RepositoryId)
+	{
+		DefaultSessionRepository.Reset(); // Will not be able to create new sessions until a mounted repository is set as default.
+		UE_LOG(LogConcert, Warning, TEXT("Unmounted the default session repository. No session will be created until a mounted repository is set as default"));
+	}
+
+	// Mark the repository as unmounted.
+	Repository.bMounted = false;
+	Repository.ProcessId = 0;
+	MountedSessionRepositories.RemoveAll([Repository](const FConcertServerSessionRepository& CandidateRepository) { return Repository.RepositoryId == CandidateRepository.RepositoryId; });
+
+	// If the repository is being dropped, the file can go away if the repository is under a known root.
+	if (bDropped && !Repository.RepositoryRootDir.IsEmpty())
+	{
+		FString RepositoryDir = Repository.RepositoryRootDir / Repository.RepositoryId.ToString();
+		ConcertUtil::DeleteDirectoryTree(*RepositoryDir);
+	}
+
+	// Update the repository database file.
+	FSystemWideCriticalSection ScopedSystemWideMutex(ConcertServerUtil::GetServerSystemMutexName());
+	FConcertServerSessionRepositoryDatabase SessionRepositoryDb;
+	ConcertServerUtil::LoadSessionRepositoryDatabase(Role, SessionRepositoryDb);
+	if (bDropped)
+	{
+		SessionRepositoryDb.Repositories.RemoveAll([Repository](const FConcertServerSessionRepository& CandidateRepository) { return Repository.RepositoryId == CandidateRepository.RepositoryId; });
+	}
+	else if (FConcertServerSessionRepository* UnmountedRepo = SessionRepositoryDb.Repositories.FindByPredicate([Repository](const FConcertServerSessionRepository& CandidateRepository) { return Repository.RepositoryId == CandidateRepository.RepositoryId; }))
+	{
+		UnmountedRepo->bMounted = false;
+		UnmountedRepo->ProcessId = 0;
+	}
+	ConcertServerUtil::SaveSessionRepositoryDatabase(Role, SessionRepositoryDb);
+
+	return true;
+}
+
 void FConcertServer::HandleDiscoverServersEvent(const FConcertMessageContext& Context)
 {
 	const FConcertAdmin_DiscoverServersEvent* Message = Context.GetMessage<FConcertAdmin_DiscoverServersEvent>();
@@ -624,6 +789,126 @@ void FConcertServer::HandleDiscoverServersEvent(const FConcertMessageContext& Co
 	}
 }
 
+TFuture<FConcertAdmin_MountSessionRepositoryResponse> FConcertServer::HandleMountSessionRepositoryRequest(const FConcertMessageContext& Context)
+{
+	const FConcertAdmin_MountSessionRepositoryRequest* Message = Context.GetMessage<FConcertAdmin_MountSessionRepositoryRequest>();
+
+	FConcertAdmin_MountSessionRepositoryResponse ResponseData;
+	if (Message->RepositoryRootDir.IsEmpty()) // Use the server configured repository root dir?
+	{
+		FConcertServerSessionRepository Repository(GetSessionRepositoriesRootDir(), Message->RepositoryId);
+		ResponseData.MountStatus = MountSessionRepository(Repository, Message->bCreateIfNotExist, /*bCleanWorkingDir*/false, /*bCleanExpiredSessions*/false, /*bSearchByPaths*/false);
+	}
+	else // Use the client supplied repository root dir.
+	{
+		FConcertServerSessionRepository Repository(Message->RepositoryRootDir, Message->RepositoryId);
+		ResponseData.MountStatus = MountSessionRepository(Repository, Message->bCreateIfNotExist, /*bCleanWorkingDir*/false, /*bCleanExpiredSessions*/false, /*bSearchByPaths*/false);
+	}
+
+	// Should the mounted repository be used as default?
+	if (Message->bAsServerDefault && ResponseData.MountStatus == EConcertSessionRepositoryMountResponseCode::Mounted)
+	{
+		const FConcertServerSessionRepository* DefaultRepository = MountedSessionRepositories.FindByPredicate([RepositoryId = Message->RepositoryId](const FConcertServerSessionRepository& Repository)
+		{
+			return Repository.RepositoryId == RepositoryId;
+		});
+
+		if (DefaultRepository)
+		{
+			DefaultSessionRepository = *DefaultRepository;
+		}
+	}
+
+	return FConcertAdmin_MountSessionRepositoryResponse::AsFuture(MoveTemp(ResponseData));
+}
+
+TFuture<FConcertAdmin_GetSessionRepositoriesResponse> FConcertServer::HandleGetSessionRepositoriesRequest(const FConcertMessageContext& Context)
+{
+	// Prevent concurrent access to the instance file.
+	FSystemWideCriticalSection ScopedSystemWideMutex(ConcertServerUtil::GetServerSystemMutexName());
+
+	// Load the global database containing all known repositories.
+	FConcertServerSessionRepositoryDatabase SessionRepositoryDb;
+	ConcertServerUtil::LoadSessionRepositoryDatabase(Role, SessionRepositoryDb);
+	bool bDatabaseUpdated = false;
+
+	// Fill up the response.
+	FConcertAdmin_GetSessionRepositoriesResponse ResponseData;
+	for (FConcertServerSessionRepository& Repository : SessionRepositoryDb.Repositories)
+	{
+		if (Repository.bMounted && !FPlatformProcess::IsApplicationRunning(Repository.ProcessId)) // Check if the state still hold.
+		{
+			Repository.bMounted = false; // Update the state.
+			Repository.ProcessId = 0;
+			bDatabaseUpdated = true;
+		}
+		ResponseData.SessionRepositories.Add(FConcertSessionRepositoryInfo{Repository.RepositoryId, Repository.bMounted});
+	}
+
+	if (bDatabaseUpdated)
+	{
+		ConcertServerUtil::SaveSessionRepositoryDatabase(Role, SessionRepositoryDb);
+	}
+
+	return FConcertAdmin_GetSessionRepositoriesResponse::AsFuture(MoveTemp(ResponseData));
+}
+
+TFuture<FConcertAdmin_DropSessionRepositoriesResponse> FConcertServer::HandleDropSessionRepositoriesRequest(const FConcertMessageContext& Context)
+{
+	const FConcertAdmin_DropSessionRepositoriesRequest* Message = Context.GetMessage<FConcertAdmin_DropSessionRepositoriesRequest>();
+	FConcertAdmin_DropSessionRepositoriesResponse ResponseData;
+
+	// Drop the repository currently loaded in-memory by this process.
+	for (const FGuid& RepositoryId : Message->RepositoryIds)
+	{
+		if (FConcertServerSessionRepository* MountedRepo = MountedSessionRepositories.FindByPredicate([RepositoryId](const FConcertServerSessionRepository& Candidate) { return Candidate.RepositoryId == RepositoryId; }))
+		{
+			if (UnmountSessionRepository(*MountedRepo, /*bDropped*/true))
+			{
+				ResponseData.DroppedRepositoryIds.Add(RepositoryId);
+			}
+		}
+	}
+
+	// Drop the repository that aren't loaded by this process, but found in the global repository database.
+	{
+		FSystemWideCriticalSection ScopedSystemWideMutex(ConcertServerUtil::GetServerSystemMutexName());
+		FConcertServerSessionRepositoryDatabase SessionRepositoryDb;
+		ConcertServerUtil::LoadSessionRepositoryDatabase(Role, SessionRepositoryDb);
+		int32 DroppedRepositoryNum = 0;
+
+		// Drop the repositories.
+		for (const FGuid& RepositoryId : Message->RepositoryIds)
+		{
+			DroppedRepositoryNum += SessionRepositoryDb.Repositories.RemoveAll([this, &RepositoryId, &ResponseData](FConcertServerSessionRepository& Repository)
+			{
+				bool bDropped = false;
+				if (Repository.RepositoryId == RepositoryId)
+				{
+					if (!Repository.bMounted || !FPlatformProcess::IsApplicationRunning(Repository.ProcessId)) // Not mounted.
+					{
+						if (!Repository.RepositoryRootDir.IsEmpty()) // This is the standard repo form and the server know what should be deleted.
+						{
+							FString ReposDir = Repository.RepositoryRootDir / Repository.RepositoryId.ToString();
+							ConcertUtil::DeleteDirectoryTree(*ReposDir);
+						}
+						ResponseData.DroppedRepositoryIds.Add(RepositoryId);
+						return true;
+					}
+				}
+				return false;
+			});
+		}
+
+		if (DroppedRepositoryNum)
+		{
+			ConcertServerUtil::SaveSessionRepositoryDatabase(Role, SessionRepositoryDb);
+		}
+	}
+
+	return FConcertAdmin_DropSessionRepositoriesResponse::AsFuture(MoveTemp(ResponseData));
+}
+
 TFuture<FConcertAdmin_SessionInfoResponse> FConcertServer::HandleCreateSessionRequest(const FConcertMessageContext& Context)
 {
 	const FConcertAdmin_CreateSessionRequest* Message = Context.GetMessage<FConcertAdmin_CreateSessionRequest>();
@@ -808,7 +1093,7 @@ FConcertAdmin_RenameSessionResponse FConcertServer::RenameSessionInternal(const
 		else
 		{
 			ArchivedSessionInfo->SessionName = Request.NewName;
-			EventSink->OnArchivedSessionRenamed(*this, Paths->GetSessionSavedDir(Request.SessionId), *ArchivedSessionInfo);
+			EventSink->OnArchivedSessionRenamed(*this, GetSessionSavedDir(Request.SessionId), *ArchivedSessionInfo);
 
 			ResponseData.ResponseCode = EConcertResponseCode::Success;
 			UE_LOG(LogConcert, Display, TEXT("User %s renamed archived session %s from %s to %s"), *Request.UserName, *ResponseData.SessionId.ToString(), *ResponseData.OldName, *Request.NewName);
@@ -889,7 +1174,6 @@ TFuture<FConcertAdmin_GetAllSessionsResponse> FConcertServer::HandleGetAllSessio
 
 	FConcertAdmin_GetAllSessionsResponse ResponseData;
 	ResponseData.LiveSessions = GetSessionsInfo();
-	ResponseData.bIncludesPreviousInstanceSessions = bExclusiveAccessToExistingSessionsAtStartup;
 	for (const auto& ArchivedSessionPair : ArchivedSessions)
 	{
 		ResponseData.ArchivedSessions.Add(ArchivedSessionPair.Value);
@@ -989,7 +1273,7 @@ bool FConcertServer::IsRequestFromSessionOwner(const TSharedPtr<IConcertServerSe
 	return false;
 }
 
-TSharedPtr<IConcertServerSession> FConcertServer::CreateLiveSession(const FConcertSessionInfo& SessionInfo)
+TSharedPtr<IConcertServerSession> FConcertServer::CreateLiveSession(const FConcertSessionInfo& SessionInfo, const FConcertServerSessionRepository& InRepository)
 {
 	check(SessionInfo.SessionId.IsValid() && !SessionInfo.SessionName.IsEmpty());
 	check(!LiveSessions.Contains(SessionInfo.SessionId) && !GetLiveSessionIdByName(SessionInfo.SessionName).IsValid());
@@ -1006,7 +1290,7 @@ TSharedPtr<IConcertServerSession> FConcertServer::CreateLiveSession(const FConce
 		LiveSessionInfo,
 		Settings->ServerSettings,
 		EndpointProvider->CreateLocalEndpoint(LiveSessionInfo.SessionName, Settings->EndpointSettings, &FConcertLogger::CreateLogger),
-		Paths->GetSessionWorkingDir(LiveSessionInfo.SessionId)
+		InRepository.GetSessionWorkingDir(LiveSessionInfo.SessionId)
 		);
 
 	LiveSessions.Add(LiveSessionInfo.SessionId, LiveSession);
@@ -1027,7 +1311,7 @@ bool FConcertServer::DestroyLiveSession(const FGuid& LiveSessionId, const bool b
 
 		if (bDeleteSessionData)
 		{
-			ConcertUtil::DeleteDirectoryTree(*Paths->GetSessionWorkingDir(LiveSessionId), *Paths->GetBaseWorkingDir());
+			ConcertUtil::DeleteDirectoryTree(*GetSessionWorkingDir(LiveSessionId));
 		}
 
 		return true;
@@ -1051,10 +1335,13 @@ FGuid FConcertServer::ArchiveLiveSession(const FGuid& LiveSessionId, const FStri
 			DestroyArchivedSession(ArchivedSessionId, /*bDeleteSessionData*/true);
 		}
 
+		// Find the live session repository to stored the archive in the same one.
+		const FConcertServerSessionRepository& SessionRepository = GetSessionRepository(LiveSession->GetSessionInfo().SessionId);
+
 		FConcertSessionInfo ArchivedSessionInfo = LiveSession->GetSessionInfo();
 		ArchivedSessionInfo.SessionId = FGuid::NewGuid();
 		ArchivedSessionInfo.SessionName = MoveTemp(ArchivedSessionName);
-		if (EventSink->ArchiveSession(*this, LiveSession.ToSharedRef(), Paths->GetSessionSavedDir(ArchivedSessionInfo.SessionId), ArchivedSessionInfo, SessionFilter))
+		if (EventSink->ArchiveSession(*this, LiveSession.ToSharedRef(), SessionRepository.GetSessionSavedDir(ArchivedSessionInfo.SessionId), ArchivedSessionInfo, SessionFilter))
 		{
 			UE_LOG(LogConcert, Display, TEXT("Live session '%s' (%s) was archived as '%s' (%s)"), *LiveSession->GetName(), *LiveSession->GetId().ToString(), *ArchivedSessionInfo.SessionName, *ArchivedSessionInfo.SessionId.ToString());
 			if (CreateArchivedSession(ArchivedSessionInfo))
@@ -1073,7 +1360,7 @@ bool FConcertServer::CreateArchivedSession(const FConcertSessionInfo& SessionInf
 	check(!ArchivedSessions.Contains(SessionInfo.SessionId) && !GetArchivedSessionIdByName(SessionInfo.SessionName).IsValid());
 
 	ArchivedSessions.Add(SessionInfo.SessionId, SessionInfo);
-	EventSink->OnArchivedSessionCreated(*this, Paths->GetSessionSavedDir(SessionInfo.SessionId), SessionInfo);
+	EventSink->OnArchivedSessionCreated(*this, GetSessionSavedDir(SessionInfo.SessionId), SessionInfo);
 
 	return true;
 }
@@ -1087,7 +1374,7 @@ bool FConcertServer::DestroyArchivedSession(const FGuid& ArchivedSessionId, cons
 
 		if (bDeleteSessionData)
 		{
-			ConcertUtil::DeleteDirectoryTree(*Paths->GetSessionSavedDir(ArchivedSessionId), *Paths->GetBaseSavedDir());
+			ConcertUtil::DeleteDirectoryTree(*GetSessionSavedDir(ArchivedSessionId));
 		}
 
 		return true;
@@ -1102,6 +1389,9 @@ TSharedPtr<IConcertServerSession> FConcertServer::RestoreArchivedSession(const F
 
 	if (const FConcertSessionInfo* ArchivedSessionInfo = ArchivedSessions.Find(ArchivedSessionId))
 	{
+		// Find the archived session repository to restore the session in the same one.
+		const FConcertServerSessionRepository& SessionRepository = GetSessionRepository(ArchivedSessionId);
+
 		FString LiveSessionName = NewSessionInfo.SessionName;
 		if (LiveSessionName.IsEmpty())
 		{
@@ -1147,10 +1437,10 @@ TSharedPtr<IConcertServerSession> FConcertServer::RestoreArchivedSession(const F
 			}
 		}
 
-		if (EventSink->RestoreSession(*this, ArchivedSessionId, Paths->GetSessionWorkingDir(LiveSessionInfo.SessionId), LiveSessionInfo, SessionFilter))
+		if (EventSink->RestoreSession(*this, ArchivedSessionId, SessionRepository.GetSessionWorkingDir(LiveSessionInfo.SessionId), LiveSessionInfo, SessionFilter))
 		{
 			UE_LOG(LogConcert, Display, TEXT("Archived session '%s' (%s) was restored as '%s' (%s)"), *ArchivedSessionInfo->SessionName, *ArchivedSessionInfo->SessionId.ToString(), *LiveSessionInfo.SessionName, *LiveSessionInfo.SessionId.ToString());
-			return CreateLiveSession(LiveSessionInfo);
+			return CreateLiveSession(LiveSessionInfo, SessionRepository);
 		}
 	}
 
diff --git a/Engine/Plugins/Developer/Concert/ConcertMain/Source/Concert/Private/ConcertServer.h b/Engine/Plugins/Developer/Concert/ConcertMain/Source/Concert/Private/ConcertServer.h
--- a/Engine/Plugins/Developer/Concert/ConcertMain/Source/Concert/Private/ConcertServer.h
+++ b/Engine/Plugins/Developer/Concert/ConcertMain/Source/Concert/Private/ConcertServer.h
@@ -6,73 +6,12 @@
 #include "IConcertServer.h"
 #include "IConcertTransportModule.h"
 #include "ConcertSettings.h"
-
+#include "ConcertServerSessionRepositories.h"
 #include "UObject/StrongObjectPtr.h"
 
 class FConcertServerSession;
 class IConcertServerEventSink;
 
-class FConcertServerPaths
-{
-public:
-	/**
-	 * Constructs paths usable by the server.
-	 * @param InRole The context in which the server exist (Disaster Recovery, MultiUsers, etc).
-	 * @param InBaseWorkingDir The base directory path where the live session data will be stored. If empty, use a default one.
-	 * @param InBaseSavedDir The base directory path where the archived session will be stored. If empty, use a default one.
-	 */
-	explicit FConcertServerPaths(const FString& InRole, const FString& InBaseWorkingDir, const FString& InBaseSavedDir);
-
-	/** Get the working directory. This is were the live sessions store their files */
-	const FString& GetWorkingDir() const
-	{
-		return WorkingDir;
-	}
-
-	/** Return the working directory for a specific session */
-	FString GetSessionWorkingDir(const FGuid& InSessionId) const
-	{
-		return WorkingDir / InSessionId.ToString();
-	}
-
-	/** Get the saved directory. This is were the archived sessions store their files */
-	const FString& GetSavedDir() const
-	{
-		return SavedDir;
-	}
-
-	/** Return the saved directory for a specific session */
-	FString GetSessionSavedDir(const FGuid& InSessionId) const
-	{
-		return SavedDir / InSessionId.ToString();
-	}
-
-	/** Returns the 'base' working directory as passed to the constructor.*/
-	FString GetBaseWorkingDir() const
-	{
-		return BaseWorkingDir;
-	}
-
-	/** Returns the 'base' saved directory as passed to the constructor.*/
-	FString GetBaseSavedDir() const
-	{
-		return BaseSavedDir;
-	}
-
-private:
-	/** Get the working directory (BaseWorkingDir/Concert/Role). This is were the active sessions store their files */
-	const FString WorkingDir;
-
-	/** Get the directory where the sessions are saved (BaseSavedDir/Concert/Role). */
-	const FString SavedDir;
-
-	/** The base working directory as passed to the constructor. */
-	const FString BaseWorkingDir;
-
-	/** The base saved directory as passed to the constructor. */
-	const FString BaseSavedDir;
-};
-
 /** Implements Concert interface */
 class FConcertServer : public IConcertServer
 {
@@ -107,9 +46,28 @@ public:
 	virtual TArray<FConcertSessionClientInfo> GetSessionClients(const FGuid& SessionId) const override;
 
 private:
+	/** Returns the root dir where the servers keeps the its internally created repositories (When the caller doesn't provide the paths). */
+	const FString& GetSessionRepositoriesRootDir() const;
+
+	const FConcertServerSessionRepository& GetSessionRepository(const FGuid& SessionId) const;
+	FString GetSessionSavedDir(const FGuid& SessionId) const;
+	FString GetSessionWorkingDir(const FGuid& SessionId) const;
+
+	EConcertSessionRepositoryMountResponseCode MountSessionRepository(FConcertServerSessionRepository& Repository, bool bCreateIfNotExist, bool bCleanWorkingDir, bool bCleanExpiredSessions, bool bSearchByPaths);
+	bool UnmountSessionRepository(FConcertServerSessionRepository& InOutRepository, bool bDropped);
+
 	/**  */
 	void HandleDiscoverServersEvent(const FConcertMessageContext& Context);
 
+	/**  */
+	TFuture<FConcertAdmin_MountSessionRepositoryResponse> HandleMountSessionRepositoryRequest(const FConcertMessageContext& Context);
+
+	/**  */
+	TFuture<FConcertAdmin_GetSessionRepositoriesResponse> HandleGetSessionRepositoriesRequest(const FConcertMessageContext& Context);
+
+	/**  */
+	TFuture<FConcertAdmin_DropSessionRepositoriesResponse> HandleDropSessionRepositoriesRequest(const FConcertMessageContext& Context);
+
 	/**  */
 	TFuture<FConcertAdmin_SessionInfoResponse> HandleCreateSessionRequest(const FConcertMessageContext& Context);
 
@@ -150,13 +108,13 @@ private:
 	TFuture<FConcertAdmin_GetSessionActivitiesResponse> HandleGetSessionActivitiesRequest(const FConcertMessageContext& Context);
 
 	/** Recover the sessions found in the working directory into live session, build the list of archived sessions and rotate them, keeping only the N most recent. */
-	void RecoverSessions();
+	void RecoverSessions(const FConcertServerSessionRepository& InRepository, bool bCleanupExpiredSessions);
 
 	/**
 	 * Migrate the live sessions from the working directory (before sessions being recovered into live one) to the archived directory.
 	 * Expected to happen at start up, before RecoverSessions(), if UConcertServerConfig::bAutoArchiveOnReboot is true.
 	 */
-	void ArchiveOfflineSessions();
+	void ArchiveOfflineSessions(const FConcertServerSessionRepository& InRepository);
 
 	/** */
 	bool CanJoinSession(const TSharedPtr<IConcertServerSession>& ServerSession, const FConcertSessionSettings& SessionSettings, const FConcertSessionVersionInfo& SessionVersionInfo, FText* OutFailureReason = nullptr);
@@ -167,7 +125,7 @@ private:
 	bool IsRequestFromSessionOwner(const TSharedPtr<IConcertServerSession>& Session, const FString& FromUserName, const FString& FromDeviceName);
 
 	/**  */
-	TSharedPtr<IConcertServerSession> CreateLiveSession(const FConcertSessionInfo& SessionInfo);
+	TSharedPtr<IConcertServerSession> CreateLiveSession(const FConcertSessionInfo& SessionInfo, const FConcertServerSessionRepository& InRepository);
 
 	/**  */
 	bool DestroyLiveSession(const FGuid& LiveSessionId, const bool bDeleteSessionData);
@@ -187,8 +145,11 @@ private:
 	/** The role of this server (eg, MultiUser, DisasterRecovery, etc) */
 	FString Role;
 
-	/** Cached root paths used by this server */
-	TUniquePtr<const FConcertServerPaths> Paths;
+	/** All session repositories used by this server. */
+	TArray<FConcertServerSessionRepository> MountedSessionRepositories;
+
+	/** The server default repository created from the server settings. */
+	TOptional<FConcertServerSessionRepository> DefaultSessionRepository;
 
 	/** The session filter to apply when auto-archiving sessions */
 	FConcertSessionFilter AutoArchiveSessionFilter;
@@ -214,6 +175,6 @@ private:
 	/** Server settings object we were configured with */
 	TStrongObjectPtr<const UConcertServerConfig> Settings;
 
-	/** Whether the server could scan/rotate/load the existing sessions (not-sharable) at startup because no other concurrent server was running. */
-	bool bExclusiveAccessToExistingSessionsAtStartup = true;
+	/** The root directory containing the session repositories. */
+	FString SessionRepositoryRootDir;
 };
diff --git a/Engine/Plugins/Developer/Concert/ConcertMain/Source/Concert/Private/ConcertServerInstanceInfo.h b/Engine/Plugins/Developer/Concert/ConcertMain/Source/Concert/Private/ConcertServerInstanceInfo.h
deleted file mode 100644
diff --git a/Engine/Plugins/Developer/Concert/ConcertMain/Source/Concert/Private/ConcertServerSessionRepositories.h b/Engine/Plugins/Developer/Concert/ConcertMain/Source/Concert/Private/ConcertServerSessionRepositories.h
new file mode 100644
--- /dev/null
+++ b/Engine/Plugins/Developer/Concert/ConcertMain/Source/Concert/Private/ConcertServerSessionRepositories.h
@@ -0,0 +1,76 @@
+// Copyright 1998-2019 Epic Games, Inc. All Rights Reserved.
+
+#pragma once
+
+#include "CoreMinimal.h"
+#include "ConcertServerSessionRepositories.generated.h"
+
+/** Keep the information about a the session repository. */
+USTRUCT()
+struct FConcertServerSessionRepository
+{
+	GENERATED_BODY()
+
+	FConcertServerSessionRepository() = default;
+
+	FConcertServerSessionRepository(const FString& InRole, FGuid InRepositoryId, const FString& InWorkingDir, const FString& InSavedDir)
+		: RepositoryId(InRepositoryId)
+		, WorkingDir(InWorkingDir.Len() ? InWorkingDir : FPaths::ProjectIntermediateDir() / InRole)
+		, SavedDir(InSavedDir.Len() ? InSavedDir : FPaths::ProjectSavedDir() / InRole)
+	{
+	}
+
+	FConcertServerSessionRepository(const FString& RootPathname, FGuid InRepositoryId)
+		: RepositoryId(InRepositoryId)
+		, RepositoryRootDir(RootPathname)
+		, WorkingDir(RootPathname / InRepositoryId.ToString() / TEXT("Live"))
+		, SavedDir(RootPathname / InRepositoryId.ToString() / TEXT("Archive"))
+	{
+	}
+
+	/** Return the working directory for a specific session */
+	FString GetSessionWorkingDir(const FGuid& InSessionId) const
+	{
+		return WorkingDir / InSessionId.ToString();
+	}
+
+	/** Return the saved directory for a specific session */
+	FString GetSessionSavedDir(const FGuid& InSessionId) const
+	{
+		return SavedDir / InSessionId.ToString();
+	}
+
+	/** The repository id.*/
+	UPROPERTY()
+	FGuid RepositoryId;
+
+	/** The repository root directory. Can be empty is the working/archive dir are not stored in a standard repository. */
+	UPROPERTY()
+	FString RepositoryRootDir;
+
+	/** This repository working directory. */
+	UPROPERTY()
+	FString WorkingDir;
+
+	/** This repository saving directory. */
+	UPROPERTY()
+	FString SavedDir;
+
+	/** The server process ID that mounted the workspace if bMounted is true. */
+	UPROPERTY()
+	int32 ProcessId = 0;
+
+	/** Whether the workspace is mounted on a server.*/
+	UPROPERTY()
+	bool bMounted = false;
+};
+
+/** Tracks the session repositories across the server instances. */
+USTRUCT()
+struct FConcertServerSessionRepositoryDatabase
+{
+	GENERATED_BODY()
+
+	UPROPERTY()
+	TArray<FConcertServerSessionRepository> Repositories;
+};
diff --git a/Engine/Plugins/Developer/Concert/ConcertMain/Source/Concert/Public/ConcertMessages.h b/Engine/Plugins/Developer/Concert/ConcertMain/Source/Concert/Public/ConcertMessages.h
--- a/Engine/Plugins/Developer/Concert/ConcertMain/Source/Concert/Public/ConcertMessages.h
+++ b/Engine/Plugins/Developer/Concert/ConcertMain/Source/Concert/Public/ConcertMessages.h
@@ -60,6 +60,19 @@ enum class EConcertSessionResponseCode : uint8
 	InvalidRequest,
 };
 
+/** Response code returned when trying to mount a session repository on the server. */
+UENUM()
+enum class EConcertSessionRepositoryMountResponseCode : uint8
+{
+	/** The repository was mounted on the invoked server. */
+	Mounted,
+	/** The repository is already mounted by another server instance. */
+	AlreadyMounted,
+	/** The repository ID could not be found in the server list. */
+	NotFound,
+};
+
+
 USTRUCT()
 struct FConcertAdmin_DiscoverServersEvent : public FConcertEndpointDiscoveryEvent
 {
@@ -92,6 +105,94 @@ struct FConcertAdmin_ServerDiscoveredEvent : public FConcertEndpointDiscoveryEve
 	EConcertServerFlags ServerFlags;
 };
 
+/** Contains information about a session repository. */
+USTRUCT()
+struct FConcertSessionRepositoryInfo
+{
+	GENERATED_BODY()
+
+	/** The repository ID.*/
+	UPROPERTY(VisibleAnywhere, Category = "Concert Message")
+	FGuid RepositoryId;
+
+	/** The mounted state of this repository. */
+	UPROPERTY(VisibleAnywhere, Category = "Concert Message")
+	bool bMounted = false;
+};
+
+/**
+ * Mount a session repository used to store session files.
+ */
+USTRUCT()
+struct FConcertAdmin_MountSessionRepositoryRequest : public FConcertRequestData
+{
+	GENERATED_BODY()
+
+	/** The repository unique Id. Must be valid. */
+	UPROPERTY(VisibleAnywhere, Category = "Concert Message")
+	FGuid RepositoryId;
+
+	/** The repository root dir (absolute path) where the repository should be found/created. If empty, the server will use its default one. */
+	UPROPERTY(VisibleAnywhere, Category = "Concert Message")
+	FString RepositoryRootDir;
+
+	/** Whether this repository is set as the default one used by server to store new sessions. */
+	UPROPERTY(VisibleAnywhere, Category = "Concert Message")
+	bool bAsServerDefault = false;
+
+	/** Whether the repository should be created if it was not found. */
+	UPROPERTY(VisibleAnywhere, Category = "Concert Message")
+	bool bCreateIfNotExist = false;
+};
+
+USTRUCT()
+struct FConcertAdmin_MountSessionRepositoryResponse : public FConcertResponseData
+{
+	GENERATED_BODY()
+
+	UPROPERTY(VisibleAnywhere, Category = "Concert Message")
+	EConcertSessionRepositoryMountResponseCode MountStatus;
+};
+
+/** Returns the list of repositories known of the server. */
+USTRUCT()
+struct FConcertAdmin_GetSessionRepositoriesRequest : public FConcertRequestData
+{
+	GENERATED_BODY()
+};
+
+USTRUCT()
+struct FConcertAdmin_GetSessionRepositoriesResponse : public FConcertResponseData
+{
+	GENERATED_BODY()
+
+	UPROPERTY(VisibleAnywhere, Category = "Concert Message")
+	TArray<FConcertSessionRepositoryInfo> SessionRepositories;
+};
+
+/**
+ * Drop one or more session repositories from the server, deleting all the contained files.
+ */
+USTRUCT()
+struct FConcertAdmin_DropSessionRepositoriesRequest : public FConcertRequestData
+{
+	GENERATED_BODY()
+
+	/** The list of repository IDs to drop. */
+	UPROPERTY(VisibleAnywhere, Category = "Concert Message")
+	TArray<FGuid> RepositoryIds;
+};
+
+USTRUCT()
+struct FConcertAdmin_DropSessionRepositoriesResponse : public FConcertResponseData
+{
+	GENERATED_BODY()
+
+	/** The list of repository IDs successfully dropped (not found == dropped). */
+	UPROPERTY(VisibleAnywhere, Category = "Concert Message")
+	TArray<FGuid> DroppedRepositoryIds;
+};
+
 USTRUCT()
 struct FConcertAdmin_GetAllSessionsRequest : public FConcertRequestData
 {
@@ -108,13 +209,6 @@ struct FConcertAdmin_GetAllSessionsResponse : public FConcertResponseData
 
 	UPROPERTY(VisibleAnywhere, Category="Concert Message")
 	TArray<FConcertSessionInfo> ArchivedSessions;
-
-	/**
-	 * True when the queried server sees the sessions created/archived during its life time and the sessions from previous (now dead) instances.
-	 * Concurrent servers instances will not see the same sessions because sessions files are not-sharable.
-	 */
-	UPROPERTY(VisibleAnywhere, Category="Concert Message")
-	bool bIncludesPreviousInstanceSessions = true;
 };
 
 USTRUCT()
diff --git a/Engine/Plugins/Developer/Concert/ConcertMain/Source/Concert/Public/ConcertSettings.h b/Engine/Plugins/Developer/Concert/ConcertMain/Source/Concert/Public/ConcertSettings.h
--- a/Engine/Plugins/Developer/Concert/ConcertMain/Source/Concert/Public/ConcertSettings.h
+++ b/Engine/Plugins/Developer/Concert/ConcertMain/Source/Concert/Public/ConcertSettings.h
@@ -155,13 +155,21 @@ public:
 	UPROPERTY(config, EditAnywhere, AdvancedDisplay, Category="Endpoint Settings", meta=(ShowOnlyInnerProperties))
 	FConcertEndpointSettings EndpointSettings;
 
-	/** The directory where the server keeps the live session files. Can be specified on the server command line with `-CONCERTWORKINGDIR=`*/
+	/** The default directory where the server keeps the live session files. Can be specified on the server command line with `-CONCERTWORKINGDIR=`*/
 	UPROPERTY(config)
 	FString WorkingDir;
 
-	/** The directory where the server keeps the archived session files. Can be specified on the server command line with `-CONCERTSAVEDDIR=`*/
+	/** The default directory where the server keeps the archived session files. Can be specified on the server command line with `-CONCERTSAVEDDIR=`*/
 	UPROPERTY(config)
 	FString ArchiveDir;
+
+	/** The root directory where the server creates new session repositories (unless the client request specifies its own root). If empty or invalid, the server will use a default. */
+	UPROPERTY(config)
+	FString SessionRepositoryRootDir;
+
+	/** If neither of WorkingDir and ArchiveDir are set, determine whether the server should mount a standard default session repository where new session will be created. */
+	UPROPERTY(config)
+	bool bMountDefaultSessionRepository = true;
 };
 
 USTRUCT()
diff --git a/Engine/Plugins/Developer/Concert/ConcertMain/Source/Concert/Public/IConcertClient.h b/Engine/Plugins/Developer/Concert/ConcertMain/Source/Concert/Public/IConcertClient.h
--- a/Engine/Plugins/Developer/Concert/ConcertMain/Source/Concert/Public/IConcertClient.h
+++ b/Engine/Plugins/Developer/Concert/ConcertMain/Source/Concert/Public/IConcertClient.h
@@ -308,7 +308,33 @@ public:
 	 */
 	virtual TSharedPtr<IConcertClientSession> GetCurrentSession() const = 0;
 
-	/** 
+	/**
+	 * Create or load a repository (directory structure) containing sessions. The server internally maps the ID to its directories.
+	 * If the server finds the repository, it will scan the corresponding directories, discover the sessions, process them and
+	 * make them visible to the other clients. If the repository is not found, but bCreateIfNotExist is true, a new repository will
+	 * be created. Repositories are persistent and visible to all server instances until dropped by the clients. Any server instance
+	 * can mount a repository unless it is already mounted by another instance.
+	 * @note Usually, a client would create a new empty repository and set it as default to populate it, then keep the ID to reload/drop it later.
+	 * @param ServerAdminEndpointId The server for which the workspace must be mounted.
+	 * @param RepositoryRootDir The base directory containing the repository to create or load. Can hold several repositories. Leave it empty to use the server one if the client/server are not running on the same machine.
+	 * @param RepositoryId A unique Id for the repository. Must be valid.
+	 * @param bCreateIfNotExist Create a new repository if the specified one doesn't exist.
+	 * @param bAsDefault Whether this repository should be set as the default one to store newly created sessions.
+	 */
+	virtual TFuture<FConcertAdmin_MountSessionRepositoryResponse> MountSessionRepository(const FGuid& ServerAdminEndpointId, const FString& RepositoryRootDir, const FGuid& RepositoryId, bool bCreateIfNotExist, bool bAsDefault = false) const = 0;
+
+	/**
+	 * Request the list of existing session repositories from the server.
+	 */
+	virtual TFuture<FConcertAdmin_GetSessionRepositoriesResponse> GetSessionRepositories(const FGuid& ServerAdminEndpointId) const = 0;
+
+	/**
+	 * Drop a set of repositories (and delete the files) from the server. All sessions stored in the unmounted repository will be unloaded and unreachable.
+	 * The server can only drop the repository if it is not mounted or has mounted it itself.
+	 */
+	virtual TFuture<FConcertAdmin_DropSessionRepositoriesResponse> DropSessionRepositories(const FGuid& ServerAdminEndpointId, const TArray<FGuid>& RepositoryIds) const = 0;
+
+	/**
 	 * Get the list of sessions available on a server
 	 * @param ServerAdminEndpointId The Id of the Concert server admin endpoint
 	 * @return A future for FConcertAdmin_GetAllSessionsResponse which contains a list of sessions
diff --git a/Engine/Plugins/Developer/Concert/DisasterRecoveryClient/Source/DisasterRecoveryClient/DisasterRecoveryClient.Build.cs b/Engine/Plugins/Developer/Concert/DisasterRecoveryClient/Source/DisasterRecoveryClient/DisasterRecoveryClient.Build.cs
--- a/Engine/Plugins/Developer/Concert/DisasterRecoveryClient/Source/DisasterRecoveryClient/DisasterRecoveryClient.Build.cs
+++ b/Engine/Plugins/Developer/Concert/DisasterRecoveryClient/Source/DisasterRecoveryClient/DisasterRecoveryClient.Build.cs
@@ -11,6 +11,7 @@ namespace UnrealBuildTool.Rules
 				{
 					"Core",
 					"CoreUObject",
+					"Engine",
 				}
 			);
 
diff --git a/Engine/Plugins/Developer/Concert/DisasterRecoveryClient/Source/DisasterRecoveryClient/Private/DisasterRecoveryClientModule.cpp b/Engine/Plugins/Developer/Concert/DisasterRecoveryClient/Source/DisasterRecoveryClient/Private/DisasterRecoveryClientModule.cpp
--- a/Engine/Plugins/Developer/Concert/DisasterRecoveryClient/Source/DisasterRecoveryClient/Private/DisasterRecoveryClientModule.cpp
+++ b/Engine/Plugins/Developer/Concert/DisasterRecoveryClient/Source/DisasterRecoveryClient/Private/DisasterRecoveryClientModule.cpp
@@ -38,7 +38,7 @@ DEFINE_LOG_CATEGORY_STATIC(LogDisasterRecovery, Log, All);
 
 
 /** Implement the Disaster Recovery module */
-class FDisasterRecoveryClientModule : public IDisasterRecoveryClientModule
+class FDisasterRecoveryClientModule : public IDisasterRecoveryClientModule, public IDisasterRecoverySessionManager
 {
 public:
 	virtual void StartupModule() override
@@ -170,35 +170,9 @@ private:
 		SetIgnoreOnRestoreState(!IsCompatibleWithOtherConcertSessions(/*SyncClientStartingSession*/nullptr, SyncClient));
 	}
 
-	FGuid GetDisasterRecoverySessionId() const
-	{
-		if (DisasterRecoveryClient)
-		{
-			if (TSharedPtr<IConcertClientSession> Session = DisasterRecoveryClient->GetConcertClient()->GetCurrentSession())
-			{
-				return Session->GetSessionInfo().SessionId;
-			}
-		}
-
-		return FGuid(); // Invalid.
-	}
-
-	/** Returns the folder where the disaster recovery service should keep the live session files (the working directory). */
-	FString GetDefaultServerWorkingDir() const
-	{
-		return FPaths::ProjectIntermediateDir() / Role / TEXT("Service");
-	}
-
-	/** Returns the folder where the disaster recovery service should keep the archived session files (the saved directory). */
-	FString GetDefaultServerArchiveDir() const
-	{
-		// Put the session data in the project dir.
-		return FPaths::ProjectSavedDir() / Role / TEXT("Service");
-	}
-
 	FString GetDisasterRecoverySessionInfoFilename() const
 	{
-		return FPaths::ProjectSavedDir() / Role / TEXT("SessionInfo.json");
+		return FPaths::ProjectSavedDir() / Role / TEXT("Sessions.json");
 	}
 
 	bool LoadDisasterRecoverySessionInfo(FDisasterRecoverySessionInfo& OutSessionInfo) const
@@ -263,28 +237,72 @@ private:
 		FDisasterRecoverySessionInfo RecoverySessionInfo;
 		LoadDisasterRecoverySessionInfo(RecoverySessionInfo);
 
-		// Is this a new session created by recovering from another one?
-		RecoverySessionInfo.Sessions.RemoveAll([](const FDisasterRecoverySession& Session)
+		auto SetAutoRestoreFlag = [](FDisasterRecoverySession& RecoverySession)
 		{
-			return Session.bAutoRestoreLastSession && Session.HostProcessId == FPlatformProcess::GetCurrentProcessId();
-		});
-
-		// Create a new session.
-		FDisasterRecoverySession& Session = RecoverySessionInfo.Sessions.AddDefaulted_GetRef();
-		Session.LastSessionName = InSession->GetSessionInfo().SessionName;
-		Session.HostProcessId = FPlatformProcess::GetCurrentProcessId();
+			// Normally, bAutoRestoreLastSession is set true here and overwritten to false when the app exits normally, but when running under the debugger, auto-restore is disabled as
+			// programmers kill applications (stop the debugger) and this should not count as a crash (unless you want to simulate crash this way - see below).
+			RecoverySession.bAutoRestoreLastSession = !FPlatformMisc::IsDebuggerPresent();
+			//Session.bAutoRestoreLastSession = true; // <- MUST BE COMMENTED BEFORE SUBMITTING: For debugging purpose only. Simulate a crash by stopping the debugger during a session.
+		};
 
-		// Normally, bAutoRestoreLastSession is set true here and overwritten to false when the app exits normally, but when running under the debugger, auto-restore is disabled as
-		// programmers kill applications (stop the debugger) and this should not count as a crash (unless you want to simulate crash this way - see below).
-		Session.bAutoRestoreLastSession = !FPlatformMisc::IsDebuggerPresent();
-		//Session.bAutoRestoreLastSession = true; // <- MUST BE COMMENTED BEFORE SUBMITTING: For debugging purpose only. Simulate a crash by stopping the debugger during a session.
+		// Is this a new session created by recovering from another one?
+		if (FDisasterRecoverySession* RestoredSession = RecoverySessionInfo.Sessions.FindByPredicate(
+			[&InSession, &SetAutoRestoreFlag](const FDisasterRecoverySession& Session) { return Session.bAutoRestoreLastSession && Session.HostProcessId == FPlatformProcess::GetCurrentProcessId(); })) // See TakeRecoverySessionOwnership().
+		{
+			RestoredSession->LastSessionName = InSession->GetSessionInfo().SessionName;
+			SetAutoRestoreFlag(*RestoredSession);
+		}
+		else // Create a new session.
+		{
+			FDisasterRecoverySession& RecoverySession = RecoverySessionInfo.Sessions.AddDefaulted_GetRef();
+			RecoverySession.RepositoryRootDir = GetSessionRepositoryRootDir();
+			RecoverySession.LastSessionName = InSession->GetSessionInfo().SessionName;
+			RecoverySession.HostProcessId = FPlatformProcess::GetCurrentProcessId();
+			RecoverySession.RepositoryId = GetSessionRepositoryId();
+			SetAutoRestoreFlag(RecoverySession);
+		}
 
 		// Save the file.
 		SaveDisasterRecoverySessionInfo(RecoverySessionInfo);
 	}
+	
+	/** Returns the number of session to keep around in the history for a given project. */
+	int32 GetMaxSessionHistorySize() const
+	{
+		return FMath::Max(0, GetDefault<UDisasterRecoverClientConfig>()->SessionHistorySize);
+	}
 
-	TPair<bool, const FConcertSessionInfo*> PinSessionToRestore(const TArray<FConcertSessionInfo>& ArchivedSessions)
+	/** Returns this client repository database root dir. */
+	virtual FString GetSessionRepositoryRootDir() const override
 	{
+		const FString& RootDir = GetDefault<UDisasterRecoverClientConfig>()->RecoverySessionDir.Path;
+		if (!RootDir.IsEmpty() && (IFileManager::Get().DirectoryExists(*RootDir) || IFileManager::Get().MakeDirectory(*RootDir, /*Tree*/true)))
+		{
+			return RootDir;
+		}
+
+		return FPaths::ProjectSavedDir() / Role / TEXT("Sessions"); // Returns the default.
+	}
+
+	/** Return the repository ID to use if a new session is created rather than recovered. */
+	virtual FGuid GetSessionRepositoryId() const override
+	{ 
+		static FGuid RepositoryId = FGuid::NewGuid();
+		return RepositoryId;
+	}
+
+	/** Among the tracked sessions, select which one is the best candidate for recovery. Concurrent Editors might be running/crashing/restoring at the same time. */
+	virtual TOptional<FDisasterRecoverySession> FindRecoverySessionCandidate(const TArray<FConcertSessionRepositoryInfo>& Repositories) override
+	{
+		// +------------------------+-----------------+-------------------+---------------+
+		// | AutoRestoreLastSession | HostProcessDead | RepositoryMounted | Deduced State |
+		// +------------------------+-----------------+-------------------+---------------+
+		// |          no            |        Any      |        Any        |  Normal Exit  | -> The session has exited properly (according to Disaster Recovery)
+		// |          yes           |        yes      |        no         |  Crashed      | -> The session is cold dead.
+		// |          yes           |        no       |        no         |  Crashing     | -> The session is crashing, but CrashReporterClientEditor detected it and shutted down DR service before the editor finished crashing.
+		// |          yes           |        Any      |        yes        |  Running      | -> The session is presumably running. Might be crashing or restoring, but as long as the repository is mounted, it is assumed running.
+		// +------------------------+-----------------+-------------------+---------------+
+	
 		// Ensure we get exclusive access to the recovery session info file.
 		FSystemWideCriticalSection SystemWideMutex(GetSystemMutexName());
 
@@ -292,67 +310,115 @@ private:
 		FDisasterRecoverySessionInfo RecoverySessionInfo;
 		LoadDisasterRecoverySessionInfo(RecoverySessionInfo);
 
-		// Checks if two running process ID are instances of same executable.
+		// Checks if two running process IDs are instances of the same executable.
 		auto IsSameExecutable = [](int32 LhsProcessId, int32 RhsProcessId)
 		{
 			return FPaths::GetPathLeaf(FPlatformProcess::GetApplicationName(LhsProcessId)) == FPaths::GetPathLeaf(FPlatformProcess::GetApplicationName(RhsProcessId));
 		};
 
 		// Returns true if the process hosting the session crashed. Note that CrashReporterClientEditor will shut down the server, archive the session and relaunch a new editor (and a new recovery client) before the
-		// previous editor finished crashing. The process hosting a session may still be alive, but not its server. In such case, the session will be unlocked and available for restoration unless another server instances already exist.
-		auto IsHostProcessCrashed = [&IsSameExecutable](const FDisasterRecoverySession& Session)
+		// previous editor finished crashing. The process hosting a session may still be alive, but not its server. In such case, the session repository will be unmounted and available for restoration unless another
+		// server instances is already restoring the session.
+		auto IsHostProcessDead = [&IsSameExecutable](const FDisasterRecoverySession& Session)
 		{
-			return Session.bAutoRestoreLastSession && (!FPlatformProcess::IsApplicationRunning(Session.HostProcessId) || !IsSameExecutable(Session.HostProcessId, FPlatformProcess::GetCurrentProcessId()));
+			return Session.bAutoRestoreLastSession && (Session.HostProcessId == 0 || !FPlatformProcess::IsApplicationRunning(Session.HostProcessId) || !IsSameExecutable(Session.HostProcessId, FPlatformProcess::GetCurrentProcessId()));
 		};
 
-		// Returns true if the session was listed as archived by the server. (Might not be visible to that server if other servers are running and use the same working/saved directories, the files are not sharable)
-		auto IsSessionArchived = [](const FDisasterRecoverySession& Session, const TArray<FConcertSessionInfo>& ArchivedSessions)
+		// Returns true if the session repository is mounted by another client/server pair preventing this client/server instance from loading it.
+		auto IsRepositoryMounted = [](const TArray<FConcertSessionRepositoryInfo>& Repositories, const FGuid& SessionRepositoryId)
 		{
-			return ArchivedSessions.FindByPredicate([&Session](const FConcertSessionInfo& MatchCandidate) { return MatchCandidate.SessionName == Session.LastSessionName; });
+			const FConcertSessionRepositoryInfo* Repository = Repositories.FindByPredicate([&SessionRepositoryId](const FConcertSessionRepositoryInfo& CandidateRepos) { return CandidateRepos.RepositoryId == SessionRepositoryId; });
+			return Repository ? Repository->bMounted : false; // Not found means not mounted -> this will likely prevent the session from being restored, but this will be gracefully handled by the FSM.
 		};
 
-		// Check in the client file for session candidate to restoration.
-		TPair<bool, const FConcertSessionInfo*> RestoreCandidate = MakeTuple(false, static_cast<const FConcertSessionInfo*>(nullptr));
+		// Sort the candidates by 'hotness'.
+		TArray<FDisasterRecoverySession*> SortedCandidates;
 		for (FDisasterRecoverySession& Session : RecoverySessionInfo.Sessions)
 		{
-			const FConcertSessionInfo* ArchivedSession = IsSessionArchived(Session, ArchivedSessions);
-			if (Session.bAutoRestoreLastSession && (ArchivedSession || IsHostProcessCrashed(Session)))
+			if (Session.bAutoRestoreLastSession)
 			{
-				// Already pinned a session to restore?
-				if (RestoreCandidate.Key && RestoreCandidate.Value)
+				if (IsRepositoryMounted(Repositories, Session.RepositoryId))
 				{
-					// User ran multiple instances of the Editor on the same project and more than one instance crashed. Arbitrary clear this one and restore the first one found only.
-					// The situation is unlikely to happen, but a future task would be to return a list of crashed sessions, pass it to the DisasterRecoveryFSM and prompt the user to select which one to recover.
-					Session.bAutoRestoreLastSession = false;
+					continue; // Two editors on the same project. This client will not be able to mount the repository, it is already mounted by another instance (which might be crashing or running, but no way to know).
 				}
-				else
+				else if (IsHostProcessDead(Session))
 				{
-					// A crashed session exist.
-					RestoreCandidate.Key = true;
-
-					// Can this session be restored (i.e. the server knowns it)?
-					if (ArchivedSession)
-					{
-						// Pin/take ownership of this session. (Will be cleared on DisasterRecoverySessionCreated() if recovery succeed)
-						RestoreCandidate.Value = ArchivedSession;
-						Session.HostProcessId = FPlatformProcess::GetCurrentProcessId();
-					}
+					SortedCandidates.Add(&Session); // The session is cold dead.
 				}
+				else // The session host process is likely crashing, but this was detected by the CRC and it unmounted the session repository.
+				{
+					SortedCandidates.Insert(&Session, 0); // Keep most recent session crash in front.
+				}
+
+				Session.HostProcessId = 0; // Host is dead or dying, clear it.
 			}
 		}
 
+		// If user ran concurrent instances of the Editor on the same project and more than one instance crashed, keep only one in front and discard the other(s). (Eventually... add code to ask the user to pick one)
+		while (SortedCandidates.Num() > 1)
+		{
+			SortedCandidates.Last()->bAutoRestoreLastSession = false;
+			SortedCandidates.Pop(/*bAllowShrinking*/false);
+		}
+
 		// Remove completed sessions (in case a session was flagged bAutoRestoreLastSession = false above)
-		RecoverySessionInfo.Sessions.RemoveAll([](const FDisasterRecoverySession& Session)
+		RecoverySessionInfo.Sessions.RemoveAll([&RecoverySessionInfo](const FDisasterRecoverySession& Session)
 		{
-			return !Session.bAutoRestoreLastSession; // Remove if the 'restore' flag is false.
+			if (!Session.bAutoRestoreLastSession) // Remove if the 'restore' flag is false.
+			{
+				RecoverySessionInfo.SessionHistory.Add(Session);
+				return true;
+			}
+			return false;
 		});
 
+		// Found a suitable candidate to restore?
+		TOptional<FDisasterRecoverySession> RestoreCandidate;
+		if (SortedCandidates.Num() == 1)
+		{
+			RestoreCandidate = *SortedCandidates[0];
+		}
+
 		// Save the recovery session info file.
 		SaveDisasterRecoverySessionInfo(RecoverySessionInfo);
 
 		return RestoreCandidate;
 	}
 
+	/** Mark this process as responsible to restore the specified session. Can take the session ownership only once the client workspace has been mounted on the server. */
+	virtual void TakeRecoverySessionOwnership(const FDisasterRecoverySession& TargetSession) override
+	{
+		// Ensure we get exclusive access to the recovery session info file.
+		FSystemWideCriticalSection SystemWideMutex(GetSystemMutexName());
+
+		// Load the session info file (if it exist)
+		FDisasterRecoverySessionInfo RecoverySessionInfo;
+		LoadDisasterRecoverySessionInfo(RecoverySessionInfo);
+
+		if (FDisasterRecoverySession* Session = RecoverySessionInfo.Sessions.FindByPredicate([&TargetSession](const FDisasterRecoverySession& Candidate) { return Candidate.RepositoryId == TargetSession.RepositoryId; }))
+		{
+			Session->HostProcessId = FPlatformProcess::GetCurrentProcessId(); // Mark this process as owner.
+
+			// Save the recovery session info file.
+			SaveDisasterRecoverySessionInfo(RecoverySessionInfo);
+		}
+	}
+
+	virtual void DiscardRecoverySession(const FDisasterRecoverySession& Session) override
+	{
+		// Ensure we get exclusive access to the recovery session info file.
+		FSystemWideCriticalSection SystemWideMutex(GetSystemMutexName());
+
+		// Load the session info file (if it exist)
+		FDisasterRecoverySessionInfo RecoverySessionInfo;
+		LoadDisasterRecoverySessionInfo(RecoverySessionInfo);
+
+		RecoverySessionInfo.Sessions.RemoveAll([Session](const FDisasterRecoverySession& Candidate) { return Session.RepositoryId == Candidate.RepositoryId; });
+
+		// Save the recovery session info file.
+		SaveDisasterRecoverySessionInfo(RecoverySessionInfo);
+	}
+
 	bool HasCandidateSessionToRestore() const
 	{
 		// Ensure we get exclusive access to the recovery session info file.
@@ -373,10 +439,48 @@ private:
 
 		// Load the session info file (if it exist)
 		FDisasterRecoverySessionInfo RecoverySessionInfo;
+		LoadDisasterRecoverySessionInfo(RecoverySessionInfo);
+		RecoverySessionInfo.SessionHistory.Append(RecoverySessionInfo.Sessions);
+		RecoverySessionInfo.Sessions.Empty();
 		SaveDisasterRecoverySessionInfo(RecoverySessionInfo);
 	}
 
-	bool SpawnDisasterRecoveryServer(const FString& ServerName) // On Linux and Mac. On Windows, it is embedded in CrashReporterClient.
+	/** Return the list of expired client workspaces that can be deleted from the server. */
+	virtual TArray<FGuid> GetExpiredSessionRepositoryIds() const override
+	{
+		// Ensure we get exclusive access to the recovery session info file.
+		FSystemWideCriticalSection SystemWideMutex(GetSystemMutexName());
+
+		// Load the session info file (if it exist)
+		FDisasterRecoverySessionInfo RecoverySessionInfo;
+		LoadDisasterRecoverySessionInfo(RecoverySessionInfo);
+		TArray<FGuid> ExpiredWorkspaceIds;
+		int32 ExpiredCount = RecoverySessionInfo.SessionHistory.Num() - GetMaxSessionHistorySize();
+		for (int32 i = 0; i < ExpiredCount; ++i)
+		{
+			ExpiredWorkspaceIds.Add(RecoverySessionInfo.SessionHistory[i].RepositoryId);
+		}
+
+		return ExpiredWorkspaceIds;
+	}
+
+	/** Invoked when client workspace were purged from the server. */
+	virtual void OnSessionRepositoryDropped(const TArray<FGuid>& PurgedWorkspaceIds) override
+	{
+		// Ensure we get exclusive access to the recovery session info file.
+		FSystemWideCriticalSection SystemWideMutex(GetSystemMutexName());
+
+		// Load the session info file (if it exist)
+		FDisasterRecoverySessionInfo RecoverySessionInfo;
+		LoadDisasterRecoverySessionInfo(RecoverySessionInfo);
+		for (const FGuid& PurgedSessionWorkspaceId : PurgedWorkspaceIds)
+		{
+			RecoverySessionInfo.SessionHistory.RemoveAll([&PurgedSessionWorkspaceId](const FDisasterRecoverySession& Session){ return PurgedSessionWorkspaceId == Session.RepositoryId; });
+		}
+		SaveDisasterRecoverySessionInfo(RecoverySessionInfo);
+	}
+
+	bool SpawnDisasterRecoveryServer(const FString& ServerName) // For Linux and Mac. On Windows, it is embedded in CrashReporterClient.
 	{
 		// Find the service path that will host the sync server
 		const FString DisasterRecoveryServicePath = GetDisasterRecoveryServicePath();
@@ -389,8 +493,6 @@ private:
 		FString DisasterRecoveryServiceCommandLine;
 		DisasterRecoveryServiceCommandLine += FString::Printf(TEXT(" -ConcertServer=\"%s\""), *ServerName);
 		DisasterRecoveryServiceCommandLine += FString::Printf(TEXT(" -EditorPID=%d"), FPlatformProcess::GetCurrentProcessId());
-		DisasterRecoveryServiceCommandLine += FString::Printf(TEXT(" -ConcertWorkingDir=\"%s\""), *GetDefaultServerWorkingDir());
-		DisasterRecoveryServiceCommandLine += FString::Printf(TEXT(" -ConcertSavedDir=\"%s\""), *GetDefaultServerArchiveDir());
 
 		// Create the service process that will host the sync server
 		DisasterRecoveryServiceHandle = FPlatformProcess::CreateProc(*DisasterRecoveryServicePath, *DisasterRecoveryServiceCommandLine, true, true, true, nullptr, 0, nullptr, nullptr, nullptr);
@@ -437,14 +539,11 @@ private:
 			ClearSessionInfoFile();
 		}
 
-		// Skip restoration if they are no candidates to restore.
-		bool bCandidateToRestore = HasCandidateSessionToRestore();
-
 		// Create and populate the client config object
 		UConcertClientConfig* ClientConfig = NewObject<UConcertClientConfig>();
 		ClientConfig->bIsHeadless = true;
 		ClientConfig->bInstallEditorToolbarButton = false;
-		ClientConfig->bAutoConnect = !bCandidateToRestore; // If there is a possible candidate to restore, don't auto connect yet, run the finite state machine to know more.
+		ClientConfig->bAutoConnect = false;
 		ClientConfig->DefaultServerURL = DisasterRecoveryServerName;
 		ClientConfig->DefaultSessionName = DisasterRecoverySessionName;
 		ClientConfig->DefaultSaveSessionAs = DisasterRecoverySessionName;
@@ -460,23 +559,15 @@ private:
 		SetIgnoreOnRestoreState(!IsCompatibleWithOtherConcertSessions(/*SyncClientStartingSession*/nullptr, /*SyncClientShuttingDownSession*/nullptr));
 
 		// If something might be recovered from a crash.
-		if (bCandidateToRestore)
+		if (HasCandidateSessionToRestore() && GUnrealEd)
 		{
-			if (GUnrealEd)
-			{
-				// Prevent the "Auto-Save" system from restoring the packages before Disaster Recovery plugin.
-				GUnrealEd->GetPackageAutoSaver().DisableRestorePromptAndDeclinePackageRecovery();
-			}
-
-			// The DisasterRecoveryFSM will call this function to select and take ownership of the session to recover.
-			auto PinSessionFn = [this](const TArray<FConcertSessionInfo>& ArchivedSessions)
-			{
-				return PinSessionToRestore(ArchivedSessions);
-			};
-
-			DisasterRecoveryUtil::StartRecovery(DisasterRecoveryClient.ToSharedRef(), PinSessionFn, /*bLiveDataOnly*/ false);
+			// Prevent the "Auto-Save" system from restoring the packages before Disaster Recovery plugin.
+			GUnrealEd->GetPackageAutoSaver().DisableRestorePromptAndDeclinePackageRecovery();
 		}
 
+		// The FSM will try to pin a session for recovery (and may fail with a toast), if no session is found, it will create a new one.
+		DisasterRecoveryUtil::StartRecovery(DisasterRecoveryClient.ToSharedRef(), *this, /*bLiveDataOnly*/ false);
+
 		return true;
 	}
 
@@ -494,9 +585,14 @@ private:
 
 			// Remove the current session from the list of sessions to track.
 			int32 ProcessId = FPlatformProcess::GetCurrentProcessId();
-			RecoverySessionInfo.Sessions.RemoveAll([ProcessId](const FDisasterRecoverySession& Session)
+			RecoverySessionInfo.Sessions.RemoveAll([ProcessId, &RecoverySessionInfo](const FDisasterRecoverySession& Session)
 			{
-				return Session.HostProcessId == ProcessId;
+				if (Session.HostProcessId == ProcessId)
+				{
+					RecoverySessionInfo.SessionHistory.Add_GetRef(Session).bAutoRestoreLastSession = false; // Push back the session in the history list.
+					return true;
+				}
+				return false;
 			});
 
 			// Write the file to disk.
diff --git a/Engine/Plugins/Developer/Concert/DisasterRecoveryClient/Source/DisasterRecoveryClient/Private/DisasterRecoveryFSM.cpp b/Engine/Plugins/Developer/Concert/DisasterRecoveryClient/Source/DisasterRecoveryClient/Private/DisasterRecoveryFSM.cpp
--- a/Engine/Plugins/Developer/Concert/DisasterRecoveryClient/Source/DisasterRecoveryClient/Private/DisasterRecoveryFSM.cpp
+++ b/Engine/Plugins/Developer/Concert/DisasterRecoveryClient/Source/DisasterRecoveryClient/Private/DisasterRecoveryFSM.cpp
@@ -16,6 +16,7 @@
 #include "SConcertSessionRecovery.h"
 #include "ConcertActivityStream.h"
 #include "GenericPlatform/GenericPlatformCrashContext.h"
+#include "DisasterRecoverySessionInfo.h"
 
 #define LOCTEXT_NAMESPACE "DisasterRecoveryFSM"
 
@@ -30,7 +31,7 @@ public:
 	static constexpr int64 MaxActivityPerRequest = 1024;
 
 public:
-	FDisasterRecoveryFSM(TSharedRef<IConcertSyncClient> InSyncClient, const DisasterRecoveryUtil::PinSessionToRestoreFunc& PinSessionToRestoreFn, bool bInLiveDataOnly);
+	FDisasterRecoveryFSM(TSharedRef<IConcertSyncClient> InSyncClient, IDisasterRecoverySessionManager& RecoverySessionManager, bool bInLiveDataOnly);
 	~FDisasterRecoveryFSM();
 
 	bool IsDone() const
@@ -58,6 +59,15 @@ private:
 
 	/** Poll the list of servers and lookup the recovery service and transit to the next state (finding the session to recover) once the recovery service is found. */
 	void LookupRecoveryServer();
+	
+	/** Tell the recovery server to drop old repositories. */
+	void DropExpiredSessionRepositories();
+
+	/** Try to find a session to recover. */
+	void SelectRecoverySession();
+
+	/** Tell the recovery server where to find this client sessions. */
+	void MountSessionRepository();
 
 	/** Poll the list of sessions from the recovery server, look up for the recovery session and transit to the next state (fetching activities) once the recovery session is found. */
 	void LookupRecoverySession();
@@ -102,12 +112,13 @@ private:
 
 private:
 	TSharedRef<IConcertSyncClient> SyncClient;
+	IDisasterRecoverySessionManager& RecoverySessionManager;
 	IConcertClientRef Client;
 	FDelegateHandle TickerHandle;
 
 	// Shared state variables.
 	FGuid RecoveryServerAdminEndpointId;
-	FString RecoverySessionName;
+	TOptional<FDisasterRecoverySession> RecoverySession;
 	FGuid RecoverySessionId;
 	TSharedPtr<FConcertActivityStream> ActivityStream;
 	TArray<TSharedPtr<FConcertClientSessionActivity>> Activities;
@@ -121,6 +132,9 @@ private:
 	// States.
 	FState EnterState;                 // State to start from.
 	FState LookupRecoveryServerState;  // Poll to find the recovery servers.
+	FState DropSessionRepositoriesState;// Delete a set of old repository (and contained sessions) from this client.
+	FState SelectRecoverySessionState; // Try to find a suitable session to recover.
+	FState MountSessionRepositoryState;// Tell the server where to discover/load/process sessions used by this client.
 	FState LookupRecoverySessionState; // Poll to find the recovery session.
 	FState FetchActivitiesState;       // Fetch all recovery session activities.
 	FState DisplayRecoveryUIState;     // Let the user view and select the recovery point.
@@ -132,20 +146,18 @@ private:
 	FState ExitState;                  // State to exit the FSM.
 	FState* CurrentState = nullptr;    // The current state.
 	FState* NextStatePending = nullptr;// State to transit to at the next Tick().
-
-	DisasterRecoveryUtil::PinSessionToRestoreFunc PinSessionToRestoreFn;
 };
 
 static TSharedPtr<FDisasterRecoveryFSM> RecoveryFSM;
 
 
-FDisasterRecoveryFSM::FDisasterRecoveryFSM(TSharedRef<IConcertSyncClient> InSyncClient, const DisasterRecoveryUtil::PinSessionToRestoreFunc& InPinSessionToRestoreFn, bool bInLiveDataOnly)
+FDisasterRecoveryFSM::FDisasterRecoveryFSM(TSharedRef<IConcertSyncClient> InSyncClient, IDisasterRecoverySessionManager& InRecoverySessionManager, bool bInLiveDataOnly)
 	: SyncClient(InSyncClient)
+	, RecoverySessionManager(InRecoverySessionManager)
 	, Client(InSyncClient->GetConcertClient())
 	, bLiveDataOnly(bInLiveDataOnly)
 	, ExitMessage(TEXT("Disaster recovery process completed successfully."))
 	, CurrentState(&EnterState)
-	, PinSessionToRestoreFn(InPinSessionToRestoreFn)
 {
 	FutureExecutionToken = MakeShared<uint8>(0); // Token use to prevent execution of async continuation if this object instance is destructed before then continuation is executed.
 
@@ -164,24 +176,27 @@ FDisasterRecoveryFSM::FDisasterRecoveryFSM(TSharedRef<IConcertSyncClient> InSync
 	};
 
 	// Implement the states, pretty much in execution order.
-	EnterState.OnExit                  = [=]()    { Statup(); };
-	LookupRecoveryServerState.OnEnter  = [this]() { Client->StartDiscovery(); };
-	LookupRecoveryServerState.OnTick   = [this]() { LookupRecoveryServer(); };
-	LookupRecoveryServerState.OnExit   = [this]() { Client->StopDiscovery(); };
-	LookupRecoverySessionState.OnEnter = [this]() { LookupRecoverySession(); };
-	FetchActivitiesState.OnEnter       = [this]() { ActivityStream = MakeShared<FConcertActivityStream>(Client, RecoveryServerAdminEndpointId, RecoverySessionId, ShouldDisplayActivityDetails()); };
-	FetchActivitiesState.OnTick        = [this]() { FetchActivities(); };
-	DisplayRecoveryUIState.OnEnter     = [this]() { DisplayRecoveryUI(); };
-	CreateAndJoinSessionState.OnEnter  = [this]() { Client->OnSessionConnectionChanged().AddRaw(this, &FDisasterRecoveryFSM::OnSessionConnectionChanged); CreateAndJoinSession(); };
-	CreateAndJoinSessionState.OnExit   = [this]() { Client->OnSessionConnectionChanged().RemoveAll(this); };
-	RestoreAndJoinSessionState.OnEnter = [this]() { Client->OnSessionConnectionChanged().AddRaw(this, &FDisasterRecoveryFSM::OnSessionConnectionChanged); RestoreAndJoinSession(); };
-	RestoreAndJoinSessionState.OnExit  = [this]() { Client->OnSessionConnectionChanged().RemoveAll(this); };
-	SynchronizeState.OnEnter           = [this]() { IDisasterRecoveryClientModule::Get().GetClient()->GetWorkspace()->OnWorkspaceSynchronized().AddRaw(this, &FDisasterRecoveryFSM::OnWorkspaceSynchronized); };
-	SynchronizeState.OnExit            = [this]() { IDisasterRecoveryClientModule::Get().GetClient()->GetWorkspace()->OnWorkspaceSynchronized().RemoveAll(this); };
-	PersistChangesState.OnEnter        = [this]() { WaitedFrameCount = 0; };
-	PersistChangesState.OnTick         = [this]() { PersistRecoveredChanges(); };
-	ErrorState.OnEnter                 = [this]() { NextStatePending = nullptr; DisplayError(); };
-	ExitState.OnEnter                  = [=]()    { Terminate(); };
+	EnterState.OnExit                    = [=]()    { Statup(); };
+	LookupRecoveryServerState.OnEnter    = [this]() { Client->StartDiscovery(); };
+	LookupRecoveryServerState.OnTick     = [this]() { LookupRecoveryServer(); };
+	LookupRecoveryServerState.OnExit     = [this]() { Client->StopDiscovery(); };
+	DropSessionRepositoriesState.OnEnter = [this]() { DropExpiredSessionRepositories(); };
+	SelectRecoverySessionState.OnEnter   = [this]() { SelectRecoverySession(); };
+	MountSessionRepositoryState.OnEnter  = [this]() { MountSessionRepository(); };
+	LookupRecoverySessionState.OnEnter   = [this]() { LookupRecoverySession(); };
+	FetchActivitiesState.OnEnter         = [this]() { ActivityStream = MakeShared<FConcertActivityStream>(Client, RecoveryServerAdminEndpointId, RecoverySessionId, ShouldDisplayActivityDetails()); };
+	FetchActivitiesState.OnTick          = [this]() { FetchActivities(); };
+	DisplayRecoveryUIState.OnEnter       = [this]() { DisplayRecoveryUI(); };
+	CreateAndJoinSessionState.OnEnter    = [this]() { Client->OnSessionConnectionChanged().AddRaw(this, &FDisasterRecoveryFSM::OnSessionConnectionChanged); CreateAndJoinSession(); };
+	CreateAndJoinSessionState.OnExit     = [this]() { Client->OnSessionConnectionChanged().RemoveAll(this); };
+	RestoreAndJoinSessionState.OnEnter   = [this]() { Client->OnSessionConnectionChanged().AddRaw(this, &FDisasterRecoveryFSM::OnSessionConnectionChanged); RestoreAndJoinSession(); };
+	RestoreAndJoinSessionState.OnExit    = [this]() { Client->OnSessionConnectionChanged().RemoveAll(this); };
+	SynchronizeState.OnEnter             = [this]() { IDisasterRecoveryClientModule::Get().GetClient()->GetWorkspace()->OnWorkspaceSynchronized().AddRaw(this, &FDisasterRecoveryFSM::OnWorkspaceSynchronized); };
+	SynchronizeState.OnExit              = [this]() { IDisasterRecoveryClientModule::Get().GetClient()->GetWorkspace()->OnWorkspaceSynchronized().RemoveAll(this); };
+	PersistChangesState.OnEnter          = [this]() { WaitedFrameCount = 0; };
+	PersistChangesState.OnTick           = [this]() { PersistRecoveredChanges(); };
+	ErrorState.OnEnter                   = [this]() { NextStatePending = nullptr; DisplayError(); };
+	ExitState.OnEnter                    = [=]()    { Terminate(); };
 
 	// Starts the finite state machine, transiting from 'EnterState' to 'LookupRecoveryServerState'.
 	TransitTo(LookupRecoveryServerState);
@@ -225,17 +240,74 @@ void FDisasterRecoveryFSM::LookupRecoveryServer()
 		if (ServerInfo.ServerName == Client->GetConfiguration()->DefaultServerURL)
 		{
 			RecoveryServerAdminEndpointId = ServerInfo.AdminEndpointId;
-			TransitTo(LookupRecoverySessionState);
+			TransitTo(DropSessionRepositoriesState);
+			return;
 		}
 	}
 }
 
-void FDisasterRecoveryFSM::LookupRecoverySession()
+void FDisasterRecoveryFSM::DropExpiredSessionRepositories()
 {
-	check(RecoveryServerAdminEndpointId.IsValid());
+	TArray<FGuid> ExpiredRepositoryIds = RecoverySessionManager.GetExpiredSessionRepositoryIds();
+	if (!ExpiredRepositoryIds.Num())
+	{
+		TransitTo(SelectRecoverySessionState);
+		return;
+	}
+	
+	// Ask the server to drop the expired client repositories.
+	TWeakPtr<uint8> ExecutionToken = FutureExecutionToken;
+	Client->DropSessionRepositories(RecoveryServerAdminEndpointId, ExpiredRepositoryIds).Next([this, ExecutionToken](const FConcertAdmin_DropSessionRepositoriesResponse& Response)
+	{
+		if (!ExecutionToken.IsValid())
+		{
+			return; // The 'this' object captured has been destructed, don't run the continuation.
+		}
+
+		// Don't care if the function fails or not, this is non-essential cleaning that can be done next time. Notify the list of dropped repositories if any.
+		if (Response.DroppedRepositoryIds.Num())
+		{
+			RecoverySessionManager.OnSessionRepositoryDropped(Response.DroppedRepositoryIds);
+		}
+
+		RequestTransitTo(SelectRecoverySessionState);
+	});
+}
 
+void FDisasterRecoveryFSM::SelectRecoverySession()
+{
 	TWeakPtr<uint8> ExecutionToken = FutureExecutionToken;
-	Client->GetServerSessions(RecoveryServerAdminEndpointId).Next([this, ExecutionToken](const FConcertAdmin_GetAllSessionsResponse& Response)
+	Client->GetSessionRepositories(RecoveryServerAdminEndpointId).Next([this, ExecutionToken](const FConcertAdmin_GetSessionRepositoriesResponse& Response)
+	{
+		if (!ExecutionToken.IsValid())
+		{
+			return; // The 'this' object captured has been destructed, don't run the continuation.
+		}
+
+		if (Response.ResponseCode == EConcertResponseCode::Success)
+		{
+			// Select which session should be restored (if any). Can only recover if the session repository is not mounted by another server instance.
+			RecoverySession = RecoverySessionManager.FindRecoverySessionCandidate(Response.SessionRepositories);
+			RequestTransitTo(MountSessionRepositoryState);
+		}
+		else
+		{
+			ErrorMessage = LOCTEXT("RepositoryQueryFailed", "Failed to retrieve repositories from the server. Recovery service will be disabled for this session.");
+			RequestTransitTo(ErrorState);
+		}
+	});
+}
+
+void FDisasterRecoveryFSM::MountSessionRepository()
+{
+	// Check if an existing repository needs to be loaded or if a new one must be created.
+	const bool bCreateIfNotExist = !RecoverySession.IsSet();
+	const FGuid RepositoryId = RecoverySession.IsSet() ? RecoverySession->RepositoryId : RecoverySessionManager.GetSessionRepositoryId();
+	const FString RepositoryRootDir = RecoverySession.IsSet() ? RecoverySession->RepositoryRootDir : RecoverySessionManager.GetSessionRepositoryRootDir(); // On restore, use the original root dir, if the user changed the root dir setting, he will likely not moves the existing sessions.
+
+	// Try to mount the repository.
+	TWeakPtr<uint8> ExecutionToken = FutureExecutionToken;
+	Client->MountSessionRepository(RecoveryServerAdminEndpointId, RepositoryRootDir, RepositoryId, bCreateIfNotExist, /*bAsDefault*/true).Next([this, ExecutionToken](const FConcertAdmin_MountSessionRepositoryResponse& Response)
 	{
 		if (!ExecutionToken.IsValid())
 		{
@@ -244,27 +316,69 @@ void FDisasterRecoveryFSM::LookupRecoverySession()
 
 		if (Response.ResponseCode == EConcertResponseCode::Success)
 		{
-			// Compare the archived sessions received from the server with the candidate sessions known by the client and select/pin a session to recover.
-			TPair<bool, const FConcertSessionInfo*> SessionToRecover = PinSessionToRestoreFn(Response.ArchivedSessions);
-			if (SessionToRecover.Value) // Took ownership of a session to restore?
+			if (RecoverySession.IsSet())
 			{
-				RecoverySessionName = SessionToRecover.Value->SessionName;
-				RecoverySessionId = SessionToRecover.Value->SessionId;
-				RequestTransitTo(FetchActivitiesState);
+				if (Response.MountStatus == EConcertSessionRepositoryMountResponseCode::Mounted)
+				{
+					RecoverySessionManager.TakeRecoverySessionOwnership(RecoverySession.GetValue()); // That client mounted the session repository first, take ownership of recovering the session.
+					RequestTransitTo(LookupRecoverySessionState); // Find the session in the mounted repository and restore it.
+				}
+				else if (Response.MountStatus == EConcertSessionRepositoryMountResponseCode::AlreadyMounted)
+				{
+					// The session repository wasn't mounted when the session was selected as candidate, but now it is mounted. This means another instance is restoring the session from that repository.
+					RequestTransitTo(SelectRecoverySessionState); // Try to select another session.
+				}
+				else
+				{
+					check(Response.MountStatus == EConcertSessionRepositoryMountResponseCode::NotFound);
+					UE_LOG(LogConcert, Error, TEXT("Failed to recover previous session. The session files were likely moved or deleted. A new session will be created."));
+					RecoverySessionManager.DiscardRecoverySession(RecoverySession.GetValue());
+					RequestTransitTo(CreateAndJoinSessionState);
+				}
+			}
+			else if (Response.MountStatus == EConcertSessionRepositoryMountResponseCode::Mounted)
+			{
+				RequestTransitTo(CreateAndJoinSessionState); // No candidate to recover. Create a new session in the new repository (set 'as default' when created).
 			}
-			else if (SessionToRecover.Key && !Response.bIncludesPreviousInstanceSessions) // The client has a candidate to restore, but the server did not report it.
+			else // NotFound/AlreadyMounted are not expected. A new repository should have been created.
 			{
-				ErrorMessage = FText::Format(LOCTEXT("RecoveryFileLocked", "Failed to recover previous session. The session is likely locked by another instance of '{0}'. Exit all Editors then relaunch to recover. Recovery service will be disabled for this session."), FText::AsCultureInvariant(DisasterRecoveryUtil::GetDisasterRecoveryServiceExeName()));
+				ErrorMessage = LOCTEXT("RepositoryMountUnexpected", "Unexpected error while mounting session repository. Recovery service will be disabled for this session.");
 				RequestTransitTo(ErrorState);
 			}
-			else if (SessionToRecover.Key && Response.bIncludesPreviousInstanceSessions)
+		}
+		else
+		{
+			ErrorMessage = LOCTEXT("RepositoryMountFailed", "Failed to mount session repository on the server. Recovery service will be disabled for this session.");
+			RequestTransitTo(ErrorState);
+		}
+	});
+}
+
+void FDisasterRecoveryFSM::LookupRecoverySession()
+{
+	check(RecoveryServerAdminEndpointId.IsValid());
+	check(RecoverySession.IsSet());
+
+	TWeakPtr<uint8> ExecutionToken = FutureExecutionToken;
+	Client->GetServerSessions(RecoveryServerAdminEndpointId).Next([this, ExecutionToken](const FConcertAdmin_GetAllSessionsResponse& Response)
+	{
+		if (!ExecutionToken.IsValid())
+		{
+			return; // The 'this' object captured has been destructed, don't run the continuation.
+		}
+
+		if (Response.ResponseCode == EConcertResponseCode::Success)
+		{
+			if (const FConcertSessionInfo* SessionToRestore = Response.ArchivedSessions.FindByPredicate([this](const FConcertSessionInfo& ArchivedSession){ return ArchivedSession.SessionName == RecoverySession->LastSessionName; }))
 			{
-				UE_LOG(LogConcert, Error, TEXT("Failed to recover previous session. The session was likely moved or deleted. A new session will be created."));
-				RequestTransitTo(CreateAndJoinSessionState); // No candidate to recover. Create a new session.
+				RecoverySessionId = SessionToRestore->SessionId;
+				RequestTransitTo(FetchActivitiesState);
 			}
 			else
 			{
-				RequestTransitTo(CreateAndJoinSessionState); // No candidate to recover. Create a new session.
+				UE_LOG(LogConcert, Error, TEXT("Failed to recover previous session. The session was likely moved or deleted. A new session will be created."));
+				RecoverySessionManager.DiscardRecoverySession(RecoverySession.GetValue());
+				RequestTransitTo(CreateAndJoinSessionState);
 			}
 		}
 		else
@@ -375,6 +489,7 @@ void FDisasterRecoveryFSM::CreateAndJoinSession()
 void FDisasterRecoveryFSM::RestoreAndJoinSession()
 {
 	check(SelectedRecoveryActivity.IsValid());
+	check(RecoverySession.IsSet());
 
 	const UConcertClientConfig* ClientConfig = Client->GetConfiguration();
 
@@ -398,7 +513,8 @@ void FDisasterRecoveryFSM::RestoreAndJoinSession()
 
 		if (Response != EConcertResponseCode::Success)
 		{
-			UE_LOG(LogConcert, Error, TEXT("Disaster recovery service failed to restore session '%s (%s)'. A new session will be created."), *RecoverySessionName, *RecoverySessionId.ToString());
+			UE_LOG(LogConcert, Error, TEXT("Disaster recovery service failed to restore session '%s (%s)'. A new session will be created."), *RecoverySession->LastSessionName, *RecoverySessionId.ToString());
+			RecoverySessionManager.DiscardRecoverySession(RecoverySession.GetValue()); // If it failed once, it has no reason to succeed later.
 			RequestTransitTo(CreateAndJoinSessionState); // Still need to have a session at the end.
 		}
 		// else -> On success callback OnSessionConnectionChanged() will transit to synchronize state.
@@ -451,11 +567,11 @@ bool FDisasterRecoveryFSM::Tick(float)
 	return true;
 }
 
-void DisasterRecoveryUtil::StartRecovery(TSharedRef<IConcertSyncClient> SyncClient, const DisasterRecoveryUtil::PinSessionToRestoreFunc& PinSessionToRestoreFn, bool bLiveDataOnly)
+void DisasterRecoveryUtil::StartRecovery(TSharedRef<IConcertSyncClient> SyncClient, IDisasterRecoverySessionManager& InRecoveryManager, bool bLiveDataOnly)
 {
 	if (!RecoveryFSM)
 	{
-		RecoveryFSM = MakeShared<FDisasterRecoveryFSM>(SyncClient, PinSessionToRestoreFn, bLiveDataOnly);
+		RecoveryFSM = MakeShared<FDisasterRecoveryFSM>(SyncClient, InRecoveryManager, bLiveDataOnly);
 	}
 }
 
diff --git a/Engine/Plugins/Developer/Concert/DisasterRecoveryClient/Source/DisasterRecoveryClient/Private/DisasterRecoveryFSM.h b/Engine/Plugins/Developer/Concert/DisasterRecoveryClient/Source/DisasterRecoveryClient/Private/DisasterRecoveryFSM.h
--- a/Engine/Plugins/Developer/Concert/DisasterRecoveryClient/Source/DisasterRecoveryClient/Private/DisasterRecoveryFSM.h
+++ b/Engine/Plugins/Developer/Concert/DisasterRecoveryClient/Source/DisasterRecoveryClient/Private/DisasterRecoveryFSM.h
@@ -8,21 +8,17 @@
 #include "Templates/Function.h"
 
 class IConcertSyncClient;
+class IDisasterRecoverySessionManager;
 
 namespace DisasterRecoveryUtil
 {
-/**
- * Function invoked by by the recovery system to pin/lock a given session to be recovered by this instance.
- */
-using PinSessionToRestoreFunc = TFunction<TPair<bool, const FConcertSessionInfo*>(const TArray<FConcertSessionInfo>& /*ArchivedSessions*/)>;
-
 /**
  * Start the recovery flow. This should be called when the previous instance of the editor exited unexpectedly.
  * @param SyncClient The client configured with the recovery server URL, the recovery session name and recovery archive name.
- * @param PinSessionToRestoreFn A function used by the process to take ownership for recovering a session.
+ * @param RecoverySessionManager The manager used to interact with the available recovery sessions.
  * @param bLiveDataOnly Filter the recovery data to only recover live transaction data. (Transaction that were not saved to disk yet)
  */
-void StartRecovery(TSharedRef<IConcertSyncClient> SyncClient, const PinSessionToRestoreFunc& PinSessionToRestoreFn, bool bLiveDataOnly);
+void StartRecovery(TSharedRef<IConcertSyncClient> SyncClient, IDisasterRecoverySessionManager& RecoverySessionManager, bool bLiveDataOnly);
 
 /**
  * End the recovery flow. This can be called to abort the recovery process.
diff --git a/Engine/Plugins/Developer/Concert/DisasterRecoveryClient/Source/DisasterRecoveryClient/Private/DisasterRecoverySessionInfo.h b/Engine/Plugins/Developer/Concert/DisasterRecoveryClient/Source/DisasterRecoveryClient/Private/DisasterRecoverySessionInfo.h
--- a/Engine/Plugins/Developer/Concert/DisasterRecoveryClient/Source/DisasterRecoveryClient/Private/DisasterRecoverySessionInfo.h
+++ b/Engine/Plugins/Developer/Concert/DisasterRecoveryClient/Source/DisasterRecoveryClient/Private/DisasterRecoverySessionInfo.h
@@ -3,6 +3,7 @@
 #pragma once
 
 #include "CoreMinimal.h"
+#include "ConcertMessages.h"
 #include "DisasterRecoverySessionInfo.generated.h"
 
 USTRUCT()
@@ -10,6 +11,15 @@ struct FDisasterRecoverySession
 {
 	GENERATED_BODY()
 
+	/** The repository ID created on the server to store that session. */
+	UPROPERTY()
+	FGuid RepositoryId;
+
+	/** The session repository root dir used to create this session. */
+	UPROPERTY()
+	FString RepositoryRootDir;
+
+	/** The name of the session. */
 	UPROPERTY()
 	FString LastSessionName;
 
@@ -17,6 +27,7 @@ struct FDisasterRecoverySession
 	UPROPERTY()
 	int32 HostProcessId = 0;
 
+	/** The flag used to determine if the session was properly ended or crashed. */
 	UPROPERTY()
 	bool bAutoRestoreLastSession = false;
 };
@@ -29,6 +40,41 @@ struct FDisasterRecoverySessionInfo
 {
 	GENERATED_BODY()
 
+	/** The list of active/crashing/crashed sessions. */
 	UPROPERTY()
 	TArray<FDisasterRecoverySession> Sessions;
+
+	/** The list of sessions kept as backup (rotated over time). */
+	UPROPERTY()
+	TArray<FDisasterRecoverySession> SessionHistory;
+};
+
+/**
+ * Abstract the management of the recovery session file.
+ */
+class IDisasterRecoverySessionManager
+{
+public:
+	~IDisasterRecoverySessionManager() = default;
+
+	/** From the set of available disaster recovery sessions, if any, select a candidate to restore. */
+	virtual TOptional<FDisasterRecoverySession> FindRecoverySessionCandidate(const TArray<FConcertSessionRepositoryInfo>& WorkspaceStats) = 0;
+	
+	/** Make this process responsible to recover the session previously selected as candidate. */
+	virtual void TakeRecoverySessionOwnership(const FDisasterRecoverySession& Session) = 0;
+
+	/** Return a list of expired session repositories that can be deleted from the server. */
+	virtual TArray<FGuid> GetExpiredSessionRepositoryIds() const = 0;
+	
+	/** Invoked when session repositories were deleted from the server. */
+	virtual void OnSessionRepositoryDropped(const TArray<FGuid>& SessionWorkspaceId) = 0;
+
+	/** Returns the session repository root directory under which the server will create the session repositories. */
+	virtual FString GetSessionRepositoryRootDir() const = 0;
+
+	/** Return the session repository ID to create if a new blank session is created (rather than restoring from an existing one). */
+	virtual FGuid GetSessionRepositoryId() const = 0;
+	
+	/** Remove the session from the list of managed session because it cannot be found/restored anymore. */
+	virtual void DiscardRecoverySession(const FDisasterRecoverySession& Session) = 0;
 };
diff --git a/Engine/Plugins/Developer/Concert/DisasterRecoveryClient/Source/DisasterRecoveryClient/Private/DisasterRecoverySettings.h b/Engine/Plugins/Developer/Concert/DisasterRecoveryClient/Source/DisasterRecoveryClient/Private/DisasterRecoverySettings.h
--- a/Engine/Plugins/Developer/Concert/DisasterRecoveryClient/Source/DisasterRecoveryClient/Private/DisasterRecoverySettings.h
+++ b/Engine/Plugins/Developer/Concert/DisasterRecoveryClient/Source/DisasterRecoveryClient/Private/DisasterRecoverySettings.h
@@ -2,6 +2,7 @@
 
 #pragma once
 
+#include "Engine/EngineTypes.h"
 #include "UObject/Object.h"
 #include "UObject/ObjectMacros.h"
 #include "DisasterRecoverySettings.generated.h"
@@ -18,4 +19,17 @@ public:
 	 */
 	UPROPERTY(config, EditAnywhere, Category="Client Settings")
 	bool bIsEnabled = true;
+
+	/**
+	 * The root directory where recovery sessions should be stored. If not set or
+	 * invalid, the recovery sessions are stored in the project saved directory.
+	 */
+	UPROPERTY(config, EditAnywhere, Category="Client Settings")
+	FDirectoryPath RecoverySessionDir;
+
+	/**
+	 * The number of recovery session to keep around in the history.
+	 */
+	UPROPERTY(config) // TODO: Expose it when the user will be able to pick up a session from a list. No need to keep sessions around until the user cannot view/select one.
+	int32 SessionHistorySize = 0;
 };
diff --git a/Engine/Source/Programs/CrashReportClient/Private/RecoveryService.cpp b/Engine/Source/Programs/CrashReportClient/Private/RecoveryService.cpp
--- a/Engine/Source/Programs/CrashReportClient/Private/RecoveryService.cpp
+++ b/Engine/Source/Programs/CrashReportClient/Private/RecoveryService.cpp
@@ -84,10 +84,8 @@ bool FRecoveryService::Startup()
 	// Setup the disaster recovery server configuration
 	UConcertServerConfig* ServerConfig = IConcertSyncServerModule::Get().ParseServerSettings(FCommandLine::Get());
 	ServerConfig->bAutoArchiveOnReboot = true; // If server crashed, was killed, etc, ensure the recovery session is archived (expected by recovery flow).
-	ServerConfig->NumSessionsToKeep = 10;
 	ServerConfig->EndpointSettings.RemoteEndpointTimeoutSeconds = 0;
-	ServerConfig->WorkingDir = GetRecoveryServiceWorkingDir();
-	ServerConfig->ArchiveDir = GetRecoveryServiceArchivedDir();
+	ServerConfig->bMountDefaultSessionRepository = false; // Let the client mount its own repository to support concurrent recovery server and prevent them from concurrently accessing non-sharable database files.
 
 	FConcertSessionFilter AutoArchiveSessionFilter;
 	AutoArchiveSessionFilter.bIncludeIgnoredActivities = true;
@@ -110,16 +108,6 @@ void FRecoveryService::Shutdown()
 	}
 }
 
-FString FRecoveryService::GetRecoveryServiceWorkingDir() const
-{
-	return FPaths::ProjectIntermediateDir();
-}
-
-FString FRecoveryService::GetRecoveryServiceArchivedDir() const
-{
-	return FPaths::ProjectSavedDir();
-}
-
 FGuid FRecoveryService::GetRecoverySessionId() const
 {
 	// As long as the Concert server is up, the session would remain live (it's going to be archived when the server shutdown or reboot).
diff --git a/Engine/Source/Programs/UnrealDisasterRecoveryService/Private/UnrealDisasterRecoveryService.cpp b/Engine/Source/Programs/UnrealDisasterRecoveryService/Private/UnrealDisasterRecoveryService.cpp
--- a/Engine/Source/Programs/UnrealDisasterRecoveryService/Private/UnrealDisasterRecoveryService.cpp
+++ b/Engine/Source/Programs/UnrealDisasterRecoveryService/Private/UnrealDisasterRecoveryService.cpp
@@ -23,7 +23,7 @@ INT32_MAIN_INT32_ARGC_TCHAR_ARGV()
 	ServerLoopInitArgs.ServiceRole = TEXT("DisasterRecovery");
 	ServerLoopInitArgs.ServiceFriendlyName = TEXT("Disaster Recovery Service");
 	ServerLoopInitArgs.ServiceAutoArchiveSessionFilter.bIncludeIgnoredActivities = true;
-    ServerLoopInitArgs.bShowConsole = false;
+	ServerLoopInitArgs.bShowConsole = false;
 
 	ServerLoopInitArgs.GetServerConfigFunc = [&EditorProcessId]() -> const UConcertServerConfig*
 	{
@@ -40,8 +40,8 @@ INT32_MAIN_INT32_ARGC_TCHAR_ARGV()
 
 		UConcertServerConfig* ServerConfig = IConcertSyncServerModule::Get().ParseServerSettings(FCommandLine::Get());
 		ServerConfig->bAutoArchiveOnReboot = true; // If server crashed, was killed, etc, ensure the recovery session is archived (expected by recovery flow).
-		ServerConfig->NumSessionsToKeep = 10;
 		ServerConfig->EndpointSettings.RemoteEndpointTimeoutSeconds = 0;
+		ServerConfig->bMountDefaultSessionRepository = false; // Let the client mount its own repository to support concurrent service and avoid them to acces the same non-sharable database files.
 		return ServerConfig;
 	};
 
diff --git a/Engine/Source/Programs/UnrealMultiUserServer/Private/UnrealMultiUserServer.cpp b/Engine/Source/Programs/UnrealMultiUserServer/Private/UnrealMultiUserServer.cpp
--- a/Engine/Source/Programs/UnrealMultiUserServer/Private/UnrealMultiUserServer.cpp
+++ b/Engine/Source/Programs/UnrealMultiUserServer/Private/UnrealMultiUserServer.cpp
@@ -9,11 +9,26 @@ IMPLEMENT_APPLICATION(UnrealMultiUserServer, "UnrealMultiUserServer");
 
 int32 RunUnrealMutilUserServer(int ArgC, TCHAR* ArgV[])
 {
+	FString Role(TEXT("MultiUser"));
 	FConcertSyncServerLoopInitArgs ServerLoopInitArgs;
 	ServerLoopInitArgs.SessionFlags = EConcertSyncSessionFlags::Default_MultiUserSession;
-	ServerLoopInitArgs.ServiceRole = TEXT("MultiUser");
+	ServerLoopInitArgs.ServiceRole = Role;
 	ServerLoopInitArgs.ServiceFriendlyName = TEXT("Multi-User Editing Server");
 
+	ServerLoopInitArgs.GetServerConfigFunc = [Role]() -> const UConcertServerConfig*
+	{
+		UConcertServerConfig* ServerConfig = IConcertSyncServerModule::Get().ParseServerSettings(FCommandLine::Get());
+		if (ServerConfig->WorkingDir.IsEmpty())
+		{
+			ServerConfig->WorkingDir = FPaths::ProjectIntermediateDir() / Role;
+		}
+		if (ServerConfig->ArchiveDir.IsEmpty())
+		{
+			ServerConfig->ArchiveDir = FPaths::ProjectSavedDir() / Role;
+		}
+		return ServerConfig;
+	};
+
 	return ConcertSyncServerLoop(ArgC, ArgV, ServerLoopInitArgs);
 }
 
