Fixed missing include file.

#rnx

diff --git a/Engine/Plugins/Developer/Concert/ConcertMain/Source/Concert/Private/ConcertServerSessionRepositories.h b/Engine/Plugins/Developer/Concert/ConcertMain/Source/Concert/Private/ConcertServerSessionRepositories.h
--- a/Engine/Plugins/Developer/Concert/ConcertMain/Source/Concert/Private/ConcertServerSessionRepositories.h
+++ b/Engine/Plugins/Developer/Concert/ConcertMain/Source/Concert/Private/ConcertServerSessionRepositories.h
@@ -3,6 +3,7 @@
 #pragma once
 
 #include "CoreMinimal.h"
+#include "Core/Public/Misc/Paths.h"
 #include "ConcertServerSessionRepositories.generated.h"
 
 /** Keep the information about a the session repository. */
