- Fixed verbosity of Multi-User/DisasterRecover endpoint discovery by adding a new log category "LogConcertDebug" defaulting to "Warning" except for the Multi-User server which default to "Log" level. The category verbosity can be adjusted from the command line as: -LogCmds="LogConcertDebug Verbose" or from a console command as: log LogConcertDebug Verbose.
  - Prevented Disaster Recovery client from discovering (and logging) all recovery services running.

diff --git a/Engine/Plugins/Developer/Concert/ConcertMain/Source/Concert/Private/ConcertClient.cpp b/Engine/Plugins/Developer/Concert/ConcertMain/Source/Concert/Private/ConcertClient.cpp
--- a/Engine/Plugins/Developer/Concert/ConcertMain/Source/Concert/Private/ConcertClient.cpp
+++ b/Engine/Plugins/Developer/Concert/ConcertMain/Source/Concert/Private/ConcertClient.cpp
@@ -1210,6 +1210,7 @@ void FConcertClient::SendDiscoverServersEvent()
 	FConcertAdmin_DiscoverServersEvent DiscoverServersEvent;
 	DiscoverServersEvent.RequiredRole = Role;
 	DiscoverServersEvent.RequiredVersion = VERSION_STRINGIFY(ENGINE_MAJOR_VERSION) TEXT(".") VERSION_STRINGIFY(ENGINE_MINOR_VERSION);
+	DiscoverServersEvent.ClientAuthenticationKey = Settings->ClientSettings.ClientAuthenticationKey;
 	ClientAdminEndpoint->PublishEvent(DiscoverServersEvent);
 }
 
diff --git a/Engine/Plugins/Developer/Concert/ConcertMain/Source/Concert/Private/ConcertLogger.cpp b/Engine/Plugins/Developer/Concert/ConcertMain/Source/Concert/Private/ConcertLogger.cpp
--- a/Engine/Plugins/Developer/Concert/ConcertMain/Source/Concert/Private/ConcertLogger.cpp
+++ b/Engine/Plugins/Developer/Concert/ConcertMain/Source/Concert/Private/ConcertLogger.cpp
@@ -194,7 +194,7 @@ void FConcertLogger::FlushLog()
 
 void FConcertLogger::LogTimeOut(const TSharedRef<IConcertMessage>& Message, const FGuid& EndpointId, const FDateTime& UtcNow)
 {
-	UE_LOG(LogConcert, Warning, TEXT("%s: Message %s timed-out sending to '%s'."), *OwnerContext.ToString(), *ConcertLoggerUtil::GetMessageString(Message), *EndpointId.ToString());
+	UE_LOG(LogConcertDebug, Warning, TEXT("%s: Message %s timed-out sending to '%s'."), *OwnerContext.ToString(), *ConcertLoggerUtil::GetMessageString(Message), *EndpointId.ToString());
 
 	if (!IsLogging())
 	{
@@ -206,7 +206,7 @@ void FConcertLogger::LogTimeOut(const TSharedRef<IConcertMessage>& Message, cons
 
 void FConcertLogger::LogSendAck(const FConcertAckData& AckData, const FGuid& DestEndpoint)
 {
-	UE_LOG(LogConcert, VeryVerbose, TEXT("%s: Acknowledgement '%s' sent to '%s' for message '%s'."), *OwnerContext.ToString(), *AckData.MessageId.ToString(), *DestEndpoint.ToString(), *AckData.SourceMessageId.ToString());
+	UE_LOG(LogConcertDebug, VeryVerbose, TEXT("%s: Acknowledgement '%s' sent to '%s' for message '%s'."), *OwnerContext.ToString(), *AckData.MessageId.ToString(), *DestEndpoint.ToString(), *AckData.SourceMessageId.ToString());
 
 	if (!IsLogging())
 	{
@@ -232,7 +232,7 @@ void FConcertLogger::LogSendAck(const FConcertAckData& AckData, const FGuid& Des
 
 void FConcertLogger::LogSendEndpointClosed(const FConcertEndpointClosedData& EndpointClosedData, const FGuid& DestEndpoint, const FDateTime& UtcNow)
 {
-	UE_LOG(LogConcert, VeryVerbose, TEXT("%s: Endpoint closure '%s' sent to '%s'."), *OwnerContext.ToString(), *EndpointClosedData.MessageId.ToString(), *DestEndpoint.ToString());
+	UE_LOG(LogConcertDebug, VeryVerbose, TEXT("%s: Endpoint closure '%s' sent to '%s'."), *OwnerContext.ToString(), *EndpointClosedData.MessageId.ToString(), *DestEndpoint.ToString());
 
 	if (!IsLogging())
 	{
@@ -258,7 +258,7 @@ void FConcertLogger::LogSendEndpointClosed(const FConcertEndpointClosedData& End
 
 void FConcertLogger::LogSendReliableHandshake(const FConcertReliableHandshakeData& ReliableHandshakeData, const FGuid& DestEndpoint, const FDateTime& UtcNow)
 {
-	UE_LOG(LogConcert, VeryVerbose, TEXT("%s: Handshake sent to '%s' (state: %s, channel: %d, index: %d)."), 
+	UE_LOG(LogConcertDebug, VeryVerbose, TEXT("%s: Handshake sent to '%s' (state: %s, channel: %d, index: %d)."), 
 		*OwnerContext.ToString(),
 		*DestEndpoint.ToString(), 
 		ConcertLoggerUtil::ReliableHandshakeStateToString(ReliableHandshakeData.HandshakeState),
@@ -290,7 +290,7 @@ void FConcertLogger::LogSendReliableHandshake(const FConcertReliableHandshakeDat
 
 void FConcertLogger::LogReceiveReliableHandshake(const FConcertReliableHandshakeData& ReliableHandshakeData, const FGuid& DestEndpoint, const FDateTime& UtcNow)
 {
-	UE_LOG(LogConcert, VeryVerbose, TEXT("%s: Handshake received from '%s' (state: %s, channel: %d, index: %d)."), 
+	UE_LOG(LogConcertDebug, VeryVerbose, TEXT("%s: Handshake received from '%s' (state: %s, channel: %d, index: %d)."), 
 		*OwnerContext.ToString(),
 		*ReliableHandshakeData.ConcertEndpointId.ToString(), 
 		ConcertLoggerUtil::ReliableHandshakeStateToString(ReliableHandshakeData.HandshakeState),
@@ -322,7 +322,7 @@ void FConcertLogger::LogReceiveReliableHandshake(const FConcertReliableHandshake
 
 void FConcertLogger::LogPublish(const TSharedRef<IConcertMessage>& Message)
 {
-	UE_LOG(LogConcert, VeryVerbose, TEXT("%s: Message %s published."), *OwnerContext.ToString(), *ConcertLoggerUtil::GetMessageString(Message));
+	UE_LOG(LogConcertDebug, VeryVerbose, TEXT("%s: Message %s published."), *OwnerContext.ToString(), *ConcertLoggerUtil::GetMessageString(Message));
 
 	if (!IsLogging())
 	{
@@ -334,7 +334,7 @@ void FConcertLogger::LogPublish(const TSharedRef<IConcertMessage>& Message)
 
 void FConcertLogger::LogSend(const TSharedRef<IConcertMessage>& Message, const FGuid& DestEndpoint)
 {
-	UE_LOG(LogConcert, VeryVerbose, TEXT("%s: Message %s sent to '%s'."), *OwnerContext.ToString(), *ConcertLoggerUtil::GetMessageString(Message), *DestEndpoint.ToString());
+	UE_LOG(LogConcertDebug, VeryVerbose, TEXT("%s: Message %s sent to '%s'."), *OwnerContext.ToString(), *ConcertLoggerUtil::GetMessageString(Message), *DestEndpoint.ToString());
 
 	if (!IsLogging())
 	{
@@ -346,7 +346,7 @@ void FConcertLogger::LogSend(const TSharedRef<IConcertMessage>& Message, const F
 
 void FConcertLogger::LogMessageReceived(const FConcertMessageContext& ConcertContext, const FGuid& DestEndpoint)
 {
-	UE_LOG(LogConcert, VeryVerbose, TEXT("%s: Message %s received from '%s'."), *OwnerContext.ToString(), *ConcertLoggerUtil::GetMessageString(ConcertContext), *ConcertContext.SenderConcertEndpointId.ToString());
+	UE_LOG(LogConcertDebug, VeryVerbose, TEXT("%s: Message %s received from '%s'."), *OwnerContext.ToString(), *ConcertLoggerUtil::GetMessageString(ConcertContext), *ConcertContext.SenderConcertEndpointId.ToString());
 
 	if (!IsLogging())
 	{
@@ -358,7 +358,7 @@ void FConcertLogger::LogMessageReceived(const FConcertMessageContext& ConcertCon
 
 void FConcertLogger::LogMessageQueued(const FConcertMessageContext& ConcertContext, const FGuid& DestEndpoint)
 {
-	UE_LOG(LogConcert, VeryVerbose, TEXT("%s: Message %s queued."), *OwnerContext.ToString(), *ConcertLoggerUtil::GetMessageString(ConcertContext));
+	UE_LOG(LogConcertDebug, VeryVerbose, TEXT("%s: Message %s queued."), *OwnerContext.ToString(), *ConcertLoggerUtil::GetMessageString(ConcertContext));
 
 	if (!IsLogging())
 	{
@@ -373,15 +373,15 @@ void FConcertLogger::LogMessageDiscarded(const FConcertMessageContext& ConcertCo
 	switch (Reason)
 	{
 	case EMessageDiscardedReason::NotRequired:
-		UE_LOG(LogConcert, VeryVerbose, TEXT("%s: Message %s discarded. Message was not required."), *OwnerContext.ToString(), *ConcertLoggerUtil::GetMessageString(ConcertContext));
+		UE_LOG(LogConcertDebug, VeryVerbose, TEXT("%s: Message %s discarded. Message was not required."), *OwnerContext.ToString(), *ConcertLoggerUtil::GetMessageString(ConcertContext));
 		break;
 
 	case EMessageDiscardedReason::AlreadyProcessed:
-		UE_LOG(LogConcert, VeryVerbose, TEXT("%s: Message %s discarded. Message was already processed."), *OwnerContext.ToString(), *ConcertLoggerUtil::GetMessageString(ConcertContext));
+		UE_LOG(LogConcertDebug, VeryVerbose, TEXT("%s: Message %s discarded. Message was already processed."), *OwnerContext.ToString(), *ConcertLoggerUtil::GetMessageString(ConcertContext));
 		break;
 
 	case EMessageDiscardedReason::UnknownEndpoint:
-		UE_LOG(LogConcert, Warning, TEXT("%s: Message %s discarded. Unknown remote endpoint '%s'."), *OwnerContext.ToString(), *ConcertLoggerUtil::GetMessageString(ConcertContext), *ConcertContext.Message->ConcertEndpointId.ToString());
+		UE_LOG(LogConcertDebug, Warning, TEXT("%s: Message %s discarded. Unknown remote endpoint '%s'."), *OwnerContext.ToString(), *ConcertLoggerUtil::GetMessageString(ConcertContext), *ConcertContext.Message->ConcertEndpointId.ToString());
 		break;
 
 	default:
@@ -399,7 +399,7 @@ void FConcertLogger::LogMessageDiscarded(const FConcertMessageContext& ConcertCo
 
 void FConcertLogger::LogProcessEvent(const FConcertMessageContext& ConcertContext, const FGuid& DestEndpoint)
 {
-	UE_LOG(LogConcert, VeryVerbose, TEXT("%s: Event %s processed."), *OwnerContext.ToString(), *ConcertLoggerUtil::GetMessageString(ConcertContext));
+	UE_LOG(LogConcertDebug, VeryVerbose, TEXT("%s: Event %s processed."), *OwnerContext.ToString(), *ConcertLoggerUtil::GetMessageString(ConcertContext));
 
 	if (!IsLogging())
 	{
@@ -411,7 +411,7 @@ void FConcertLogger::LogProcessEvent(const FConcertMessageContext& ConcertContex
 
 void FConcertLogger::LogProcessRequest(const FConcertMessageContext& ConcertContext, const FGuid& DestEndpoint)
 {
-	UE_LOG(LogConcert, VeryVerbose, TEXT("%s: Request %s processed."), *OwnerContext.ToString(), *ConcertLoggerUtil::GetMessageString(ConcertContext));
+	UE_LOG(LogConcertDebug, VeryVerbose, TEXT("%s: Request %s processed."), *OwnerContext.ToString(), *ConcertLoggerUtil::GetMessageString(ConcertContext));
 
 	if (!IsLogging())
 	{
@@ -424,7 +424,7 @@ void FConcertLogger::LogProcessRequest(const FConcertMessageContext& ConcertCont
 void FConcertLogger::LogProcessResponse(const FConcertMessageContext& ConcertContext, const FGuid& DestEndpoint) 
 {
 	const FConcertResponseData* ResponseMessage = ConcertContext.GetMessage<FConcertResponseData>();
-	UE_LOG(LogConcert, VeryVerbose, TEXT("%s: Response %s processed for request '%s'."), *OwnerContext.ToString(), *ConcertLoggerUtil::GetMessageString(ConcertContext), *ResponseMessage->RequestMessageId.ToString());
+	UE_LOG(LogConcertDebug, VeryVerbose, TEXT("%s: Response %s processed for request '%s'."), *OwnerContext.ToString(), *ConcertLoggerUtil::GetMessageString(ConcertContext), *ResponseMessage->RequestMessageId.ToString());
 
 	if (!IsLogging())
 	{
@@ -437,7 +437,7 @@ void FConcertLogger::LogProcessResponse(const FConcertMessageContext& ConcertCon
 void FConcertLogger::LogProcessAck(const FConcertMessageContext& ConcertContext, const FGuid& DestEndpoint)
 {
 	const FConcertAckData* AckMessage = ConcertContext.GetMessage<FConcertAckData>();
-	UE_LOG(LogConcert, VeryVerbose, TEXT("%s: Acknowledgement '%s' processed for message '%s'."), *OwnerContext.ToString(), *AckMessage->MessageId.ToString(), *AckMessage->SourceMessageId.ToString());
+	UE_LOG(LogConcertDebug, VeryVerbose, TEXT("%s: Acknowledgement '%s' processed for message '%s'."), *OwnerContext.ToString(), *AckMessage->MessageId.ToString(), *AckMessage->SourceMessageId.ToString());
 
 	if (!IsLogging())
 	{
@@ -449,7 +449,7 @@ void FConcertLogger::LogProcessAck(const FConcertMessageContext& ConcertContext,
 
 void FConcertLogger::LogRemoteEndpointDiscovery(const FConcertMessageContext& ConcertContext, const FGuid& DestEndpoint)
 {
-	UE_LOG(LogConcert, Display, TEXT("%s: Remote endpoint '%s' discovered."), *OwnerContext.ToString(), *ConcertContext.SenderConcertEndpointId.ToString());
+	UE_LOG(LogConcertDebug, Display, TEXT("%s: Remote endpoint '%s' discovered."), *OwnerContext.ToString(), *ConcertContext.SenderConcertEndpointId.ToString());
 
 	if (!IsLogging())
 	{
@@ -461,7 +461,7 @@ void FConcertLogger::LogRemoteEndpointDiscovery(const FConcertMessageContext& Co
 
 void FConcertLogger::LogRemoteEndpointTimeOut(const FGuid& EndpointId, const FDateTime& UtcNow)
 {
-	UE_LOG(LogConcert, Display, TEXT("%s: Remote endpoint '%s' timed-out."), *OwnerContext.ToString(), *EndpointId.ToString());
+	UE_LOG(LogConcertDebug, Display, TEXT("%s: Remote endpoint '%s' timed-out."), *OwnerContext.ToString(), *EndpointId.ToString());
 
 	if (!IsLogging())
 	{
@@ -487,7 +487,7 @@ void FConcertLogger::LogRemoteEndpointTimeOut(const FGuid& EndpointId, const FDa
 
 void FConcertLogger::LogRemoteEndpointClosure(const FGuid& EndpointId, const FDateTime& UtcNow)
 {
-	UE_LOG(LogConcert, Display, TEXT("%s: Remote endpoint '%s' closed by remote peer."), *OwnerContext.ToString(), *EndpointId.ToString());
+	UE_LOG(LogConcertDebug, Display, TEXT("%s: Remote endpoint '%s' closed by remote peer."), *OwnerContext.ToString(), *EndpointId.ToString());
 
 	if (!IsLogging())
 	{
diff --git a/Engine/Plugins/Developer/Concert/ConcertMain/Source/Concert/Private/ConcertServer.cpp b/Engine/Plugins/Developer/Concert/ConcertMain/Source/Concert/Private/ConcertServer.cpp
--- a/Engine/Plugins/Developer/Concert/ConcertMain/Source/Concert/Private/ConcertServer.cpp
+++ b/Engine/Plugins/Developer/Concert/ConcertMain/Source/Concert/Private/ConcertServer.cpp
@@ -781,11 +781,14 @@ void FConcertServer::HandleDiscoverServersEvent(const FConcertMessageContext& Co
 
 	if (ServerAdminEndpoint.IsValid() && Message->RequiredRole == Role && Message->RequiredVersion == VERSION_STRINGIFY(ENGINE_MAJOR_VERSION) TEXT(".") VERSION_STRINGIFY(ENGINE_MINOR_VERSION))
 	{
-		FConcertAdmin_ServerDiscoveredEvent DiscoveryInfo;
-		DiscoveryInfo.ServerName = ServerInfo.ServerName;
-		DiscoveryInfo.InstanceInfo = ServerInfo.InstanceInfo;
-		DiscoveryInfo.ServerFlags = ServerInfo.ServerFlags;
-		ServerAdminEndpoint->SendEvent(DiscoveryInfo, Context.SenderConcertEndpointId);
+		if (Settings->AuthorizedClientKeys.Num() == 0 || Settings->AuthorizedClientKeys.Contains(Message->ClientAuthenticationKey)) // Can the client discover this server?
+		{
+			FConcertAdmin_ServerDiscoveredEvent DiscoveryInfo;
+			DiscoveryInfo.ServerName = ServerInfo.ServerName;
+			DiscoveryInfo.InstanceInfo = ServerInfo.InstanceInfo;
+			DiscoveryInfo.ServerFlags = ServerInfo.ServerFlags;
+			ServerAdminEndpoint->SendEvent(DiscoveryInfo, Context.SenderConcertEndpointId);
+		}
 	}
 }
 
diff --git a/Engine/Plugins/Developer/Concert/ConcertMain/Source/Concert/Public/ConcertMessages.h b/Engine/Plugins/Developer/Concert/ConcertMain/Source/Concert/Public/ConcertMessages.h
--- a/Engine/Plugins/Developer/Concert/ConcertMain/Source/Concert/Public/ConcertMessages.h
+++ b/Engine/Plugins/Developer/Concert/ConcertMain/Source/Concert/Public/ConcertMessages.h
@@ -85,6 +85,10 @@ struct FConcertAdmin_DiscoverServersEvent : public FConcertEndpointDiscoveryEven
 	/** The required version of the server (eg, 4.22, 4.23, etc) */
 	UPROPERTY(VisibleAnywhere, Category = "Concert Message")
 	FString RequiredVersion;
+
+	/** If a server was configured to restrict access to specific client(s), it will search for this key in its list of authorized keys.*/
+	UPROPERTY(VisibleAnywhere, Category = "Concert Message")
+	FString ClientAuthenticationKey;
 };
 
 USTRUCT()
diff --git a/Engine/Plugins/Developer/Concert/ConcertMain/Source/Concert/Public/ConcertSettings.h b/Engine/Plugins/Developer/Concert/ConcertMain/Source/Concert/Public/ConcertSettings.h
--- a/Engine/Plugins/Developer/Concert/ConcertMain/Source/Concert/Public/ConcertSettings.h
+++ b/Engine/Plugins/Developer/Concert/ConcertMain/Source/Concert/Public/ConcertSettings.h
@@ -134,6 +134,12 @@ public:
 	UPROPERTY(config, EditAnywhere, Category="Session Settings")
 	FString DefaultSessionName;
 
+	/** 
+	 * A set of keys identifying the clients that can discover and access the server. If empty, the server can be discovered and used by any clients.
+	 */
+	UPROPERTY(config)
+	TSet<FString> AuthorizedClientKeys;
+
 	/**
 	 * Name of the default session to restore on the server.
 	 * Set the name of the desired save to restore its content in your session.
@@ -221,6 +227,10 @@ struct FConcertClientSettings
 	/** Array of tags that can be used for grouping and categorizing. */
 	UPROPERTY(config, EditAnywhere, AdvancedDisplay, Category = "Client Settings")
 	TArray<FName> Tags;
+
+	/** A key used to identify the clients during server discovery. If the server was configured to restrict access, the client key must be know of the server. Can be left empty. */
+	UPROPERTY(config)
+	FString ClientAuthenticationKey;
 };
 
 UCLASS(config=Engine)
diff --git a/Engine/Plugins/Developer/Concert/ConcertMain/Source/ConcertTransport/Private/ConcertLogGlobal.cpp b/Engine/Plugins/Developer/Concert/ConcertMain/Source/ConcertTransport/Private/ConcertLogGlobal.cpp
--- a/Engine/Plugins/Developer/Concert/ConcertMain/Source/ConcertTransport/Private/ConcertLogGlobal.cpp
+++ b/Engine/Plugins/Developer/Concert/ConcertMain/Source/ConcertTransport/Private/ConcertLogGlobal.cpp
@@ -3,3 +3,4 @@
 #include "ConcertLogGlobal.h"
 
 DEFINE_LOG_CATEGORY(LogConcert);
+DEFINE_LOG_CATEGORY(LogConcertDebug);
diff --git a/Engine/Plugins/Developer/Concert/ConcertMain/Source/ConcertTransport/Public/ConcertLogGlobal.h b/Engine/Plugins/Developer/Concert/ConcertMain/Source/ConcertTransport/Public/ConcertLogGlobal.h
--- a/Engine/Plugins/Developer/Concert/ConcertMain/Source/ConcertTransport/Public/ConcertLogGlobal.h
+++ b/Engine/Plugins/Developer/Concert/ConcertMain/Source/ConcertTransport/Public/ConcertLogGlobal.h
@@ -4,4 +4,10 @@
 
 #include "Logging/LogMacros.h"
 
+// The default debug verbosity level can be set as a GlobalDefinitions in an executable Target.cs build file to adapt to this executable context.
+#ifndef UE_LOG_CONCERT_DEBUG_VERBOSITY_LEVEL
+#define UE_LOG_CONCERT_DEBUG_VERBOSITY_LEVEL Warning
+#endif
+
 CONCERTTRANSPORT_API DECLARE_LOG_CATEGORY_EXTERN(LogConcert, Log, All);
+CONCERTTRANSPORT_API DECLARE_LOG_CATEGORY_EXTERN(LogConcertDebug, UE_LOG_CONCERT_DEBUG_VERBOSITY_LEVEL, All);
diff --git a/Engine/Plugins/Developer/Concert/DisasterRecoveryClient/Source/DisasterRecoveryClient/Private/DisasterRecoveryClientModule.cpp b/Engine/Plugins/Developer/Concert/DisasterRecoveryClient/Source/DisasterRecoveryClient/Private/DisasterRecoveryClientModule.cpp
--- a/Engine/Plugins/Developer/Concert/DisasterRecoveryClient/Source/DisasterRecoveryClient/Private/DisasterRecoveryClientModule.cpp
+++ b/Engine/Plugins/Developer/Concert/DisasterRecoveryClient/Source/DisasterRecoveryClient/Private/DisasterRecoveryClientModule.cpp
@@ -266,7 +266,7 @@ private:
 		SaveDisasterRecoverySessionInfo(RecoverySessionInfo);
 	}
 	
-	/** Returns the number of session to keep around in the history for a given project. */
+	/** Returns the number of sessions to keep around in the history for a given project. */
 	int32 GetMaxSessionHistorySize() const
 	{
 		return FMath::Max(0, GetDefault<UDisasterRecoverClientConfig>()->SessionHistorySize);
@@ -548,6 +548,7 @@ private:
 		ClientConfig->DefaultSaveSessionAs = DisasterRecoverySessionName;
 		//ClientConfig->ClientSettings.DiscoveryTimeoutSeconds = 0;
 		ClientConfig->EndpointSettings.RemoteEndpointTimeoutSeconds = 0;
+		ClientConfig->ClientSettings.ClientAuthenticationKey = DisasterRecoveryServerName; // The discovery server adds its own server name to the list of authorized client keys, use that key to discover the server.
 
 		// Create the recovery session and auto-join it if there is nothing to recover.
 		DisasterRecoveryClient = IConcertSyncClientModule::Get().CreateClient(Role);
diff --git a/Engine/Source/Programs/CrashReportClient/Private/RecoveryService.cpp b/Engine/Source/Programs/CrashReportClient/Private/RecoveryService.cpp
--- a/Engine/Source/Programs/CrashReportClient/Private/RecoveryService.cpp
+++ b/Engine/Source/Programs/CrashReportClient/Private/RecoveryService.cpp
@@ -86,6 +86,7 @@ bool FRecoveryService::Startup()
 	ServerConfig->bAutoArchiveOnReboot = true; // If server crashed, was killed, etc, ensure the recovery session is archived (expected by recovery flow).
 	ServerConfig->EndpointSettings.RemoteEndpointTimeoutSeconds = 0;
 	ServerConfig->bMountDefaultSessionRepository = false; // Let the client mount its own repository to support concurrent recovery server and prevent them from concurrently accessing non-sharable database files.
+	ServerConfig->AuthorizedClientKeys.Add(ServerConfig->ServerName); // The disaster recovery client is configured to use the unique server name as key to identify itself.
 
 	FConcertSessionFilter AutoArchiveSessionFilter;
 	AutoArchiveSessionFilter.bIncludeIgnoredActivities = true;
diff --git a/Engine/Source/Programs/UnrealDisasterRecoveryService/Private/UnrealDisasterRecoveryService.cpp b/Engine/Source/Programs/UnrealDisasterRecoveryService/Private/UnrealDisasterRecoveryService.cpp
--- a/Engine/Source/Programs/UnrealDisasterRecoveryService/Private/UnrealDisasterRecoveryService.cpp
+++ b/Engine/Source/Programs/UnrealDisasterRecoveryService/Private/UnrealDisasterRecoveryService.cpp
@@ -41,7 +41,8 @@ INT32_MAIN_INT32_ARGC_TCHAR_ARGV()
 		UConcertServerConfig* ServerConfig = IConcertSyncServerModule::Get().ParseServerSettings(FCommandLine::Get());
 		ServerConfig->bAutoArchiveOnReboot = true; // If server crashed, was killed, etc, ensure the recovery session is archived (expected by recovery flow).
 		ServerConfig->EndpointSettings.RemoteEndpointTimeoutSeconds = 0;
-		ServerConfig->bMountDefaultSessionRepository = false; // Let the client mount its own repository to support concurrent service and avoid them to acces the same non-sharable database files.
+		ServerConfig->bMountDefaultSessionRepository = false; // Let the client mount its own repository to support concurrent service and avoid them to access the same non-sharable database files.
+		ServerConfig->AuthorizedClientKeys.Add(ServerConfig->ServerName); // The disaster recovery client is configured to use the unique server name as key to identify itself.
 		return ServerConfig;
 	};
 
diff --git a/Engine/Source/Programs/UnrealMultiUserServer/UnrealMultiUserServer.Target.cs b/Engine/Source/Programs/UnrealMultiUserServer/UnrealMultiUserServer.Target.cs
--- a/Engine/Source/Programs/UnrealMultiUserServer/UnrealMultiUserServer.Target.cs
+++ b/Engine/Source/Programs/UnrealMultiUserServer/UnrealMultiUserServer.Target.cs
@@ -24,5 +24,7 @@ public class UnrealMultiUserServerTarget : TargetRules
 
 		// The Multi-User server is meant to be a console application (no window), but on MacOS, to get a proper log console, a full application must be built.
 		bIsBuildingConsoleApplication = Target.Platform != UnrealTargetPlatform.Mac;
+
+		GlobalDefinitions.Add("UE_LOG_CONCERT_DEBUG_VERBOSITY_LEVEL=Log");
 	}
 }
