Fixed concert client that did not cancel in-flight connection tasks as expected if IConcertClient::DisconnectSession() was called before the session was created. For example if the CreateSession() and DisconnectSession() were called in a very short interval.

diff --git a/Engine/Plugins/Developer/Concert/ConcertMain/Source/Concert/Private/ConcertClient.cpp b/Engine/Plugins/Developer/Concert/ConcertMain/Source/Concert/Private/ConcertClient.cpp
--- a/Engine/Plugins/Developer/Concert/ConcertMain/Source/Concert/Private/ConcertClient.cpp
+++ b/Engine/Plugins/Developer/Concert/ConcertMain/Source/Concert/Private/ConcertClient.cpp
@@ -939,6 +939,9 @@ void FConcertClient::DisconnectSession()
 	// We don't want the client to get automatically reconnected to it's default session
 	AutoConnection.Reset();
 	InternalDisconnectSession();
+
+	// If async connection tasks were in-flight, cancel them.
+	PendingConnection.Reset();
 }
 
 void FConcertClient::ResumeSession()
