Fixed Concert server using a dangling reference when unmounting a repository.
Fixed Concert server to unregister left over request handler on shutdown.
Fixed Concert server handler to report a repository as mounted if the requesting client already mounted it in a previous request.
Updated concert server to clean up its repository database when a repository doesn't exist anymore on disk.

diff --git a/Engine/Plugins/Developer/Concert/ConcertMain/Source/Concert/Private/ConcertServer.cpp b/Engine/Plugins/Developer/Concert/ConcertMain/Source/Concert/Private/ConcertServer.cpp
--- a/Engine/Plugins/Developer/Concert/ConcertMain/Source/Concert/Private/ConcertServer.cpp
+++ b/Engine/Plugins/Developer/Concert/ConcertMain/Source/Concert/Private/ConcertServer.cpp
@@ -189,6 +189,21 @@ void FConcertServer::Startup()
 			FConcertServerSessionRepositoryDatabase SessionRepositoryDb;
 			ConcertServerUtil::LoadSessionRepositoryDatabase(Role, SessionRepositoryDb);
 
+			// Unmap repositories that doesn't exist anymore on disk (were deleted manually).
+			int RemovedNum = SessionRepositoryDb.Repositories.RemoveAll([](const FConcertServerSessionRepository& RemoveCandidate)
+			{
+				if (!RemoveCandidate.RepositoryRootDir.IsEmpty()) // Under a single standard root?
+				{
+					FString Pathname = RemoveCandidate.RepositoryRootDir / RemoveCandidate.RepositoryId.ToString();
+					return !IFileManager::Get().DirectoryExists(*Pathname);
+				}
+				return false; // Not under a single root (Multi-User backward compatibility mode) leave it.
+			});
+			if (RemovedNum)
+			{
+				ConcertServerUtil::SaveSessionRepositoryDatabase(Role, SessionRepositoryDb);
+			}
+
 			// Walk the root directory containing the server managed repositories and find those that aren't mapped anymore.
 			TArray<FString> ExpiredDirectories;
 			IFileManager::Get().IterateDirectory(*GetSessionRepositoriesRootDir(), [this, &SessionRepositoryDb, &ExpiredDirectories](const TCHAR* Pathname, bool bIsDirectory)
@@ -199,7 +214,7 @@ void FConcertServer::Startup()
 					FString RootReposDir = GetSessionRepositoriesRootDir();
 					if (!SessionRepositoryDb.Repositories.ContainsByPredicate([&RootReposDir, Pathname](const FConcertServerSessionRepository& Repository) { return  RootReposDir / Repository.RepositoryId.ToString() == Pathname; }))
 					{
-						ExpiredDirectories.Emplace(Pathname); // The visited directory was not found in the list of repositories.
+						ExpiredDirectories.Emplace(Pathname); // The visited directory was not found in the list of mapped repositories.
 					}
 				}
 				return true;
@@ -259,6 +274,10 @@ void FConcertServer::Shutdown()
 		ServerAdminEndpoint->UnregisterRequestHandler<FConcertAdmin_GetLiveSessionsRequest>();
 		ServerAdminEndpoint->UnregisterRequestHandler<FConcertAdmin_GetArchivedSessionsRequest>();
 		ServerAdminEndpoint->UnregisterRequestHandler<FConcertAdmin_GetSessionClientsRequest>();
+		ServerAdminEndpoint->UnregisterRequestHandler<FConcertAdmin_GetSessionActivitiesRequest>();
+		ServerAdminEndpoint->UnregisterRequestHandler<FConcertAdmin_MountSessionRepositoryRequest>();
+		ServerAdminEndpoint->UnregisterRequestHandler<FConcertAdmin_GetSessionRepositoriesRequest>();
+		ServerAdminEndpoint->UnregisterRequestHandler<FConcertAdmin_DropSessionRepositoriesRequest>();
 
 		ServerAdminEndpoint.Reset();
 	}
@@ -665,7 +684,7 @@ EConcertSessionRepositoryMountResponseCode FConcertServer::MountSessionRepositor
 		if (FConcertServerSessionRepository* ExistingRepository = SessionRepositoryDb.Repositories.FindByPredicate(
 			[&Repository, bSearchByPaths](const FConcertServerSessionRepository& Candidate){ return bSearchByPaths ? Candidate.WorkingDir == Repository.WorkingDir && Candidate.SavedDir == Repository.SavedDir : Candidate.RepositoryId == Repository.RepositoryId; }))
 		{
-			if (!ExistingRepository->bMounted || !FPlatformProcess::IsApplicationRunning(ExistingRepository->ProcessId))
+			if (!ExistingRepository->bMounted || !FPlatformProcess::IsApplicationRunning(ExistingRepository->ProcessId)) // Not mounted or mounted by a dead process.
 			{
 				check(Repository.RepositoryRootDir == ExistingRepository->RepositoryRootDir) // The client changed the root dir?
 				ExistingRepository->bMounted = true;
@@ -674,6 +693,11 @@ EConcertSessionRepositoryMountResponseCode FConcertServer::MountSessionRepositor
 				MountedSessionRepositories.Add(Repository);
 				ConcertServerUtil::SaveSessionRepositoryDatabase(Role, SessionRepositoryDb);
 			}
+			else if (ExistingRepository->ProcessId == FPlatformProcess::GetCurrentProcessId() &&
+				MountedSessionRepositories.ContainsByPredicate([&Repository](const FConcertServerSessionRepository& MatchCandidate){ return MatchCandidate.RepositoryId == Repository.RepositoryId; })) // Already mounted by this process?
+			{
+				return EConcertSessionRepositoryMountResponseCode::Mounted;
+			}
 			else
 			{
 				return EConcertSessionRepositoryMountResponseCode::AlreadyMounted; // Already mounted by another process, cannot mount it.
@@ -710,10 +734,18 @@ EConcertSessionRepositoryMountResponseCode FConcertServer::MountSessionRepositor
 	return EConcertSessionRepositoryMountResponseCode::Mounted;
 }
 
-bool FConcertServer::UnmountSessionRepository(FConcertServerSessionRepository& Repository, bool bDropped)
+bool FConcertServer::UnmountSessionRepository(const FGuid& RepositoryId, bool bDropped)
 {
-	check(Repository.bMounted && Repository.ProcessId == FPlatformProcess::GetCurrentProcessId());
-	check(MountedSessionRepositories.FindByPredicate([Repository](const FConcertServerSessionRepository& Candidate) { return Candidate.RepositoryId == Repository.RepositoryId; })); // Should be found in the list of mounted repositories.
+	// Search the repository in the list of mounted repositories.
+	int32 Index = MountedSessionRepositories.IndexOfByPredicate([&RepositoryId](const FConcertServerSessionRepository& MatchCandidate) { return RepositoryId == MatchCandidate.RepositoryId; });
+	if (Index == INDEX_NONE)
+	{
+		return false; // Not mounted by this process.
+	}
+
+	FConcertServerSessionRepository& Repository = MountedSessionRepositories[Index];
+	check(Repository.bMounted); // Must be mounted if present in the 'mounted' list.
+	check(Repository.ProcessId == FPlatformProcess::GetCurrentProcessId()); // Must be mounted by this process to be in the list.
 
 	// Unload the live sessions hosted in that repository.
 	TArray<FGuid> LiveSessionIds;
@@ -721,9 +753,9 @@ bool FConcertServer::UnmountSessionRepository(FConcertServerSessionRepository& R
 	for (const FGuid& LiveSessionId : LiveSessionIds)
 	{
 		const FConcertServerSessionRepository& SessionRepository = GetSessionRepository(LiveSessionId);
-		if (SessionRepository.RepositoryId == Repository.RepositoryId)
+		if (SessionRepository.RepositoryId == RepositoryId)
 		{
-			DestroyLiveSession(LiveSessionId, /*bDeleteSessionData*/false);
+			DestroyLiveSession(LiveSessionId, /*bDeleteSessionData*/bDropped);
 		}
 	}
 
@@ -733,44 +765,50 @@ bool FConcertServer::UnmountSessionRepository(FConcertServerSessionRepository& R
 	for (const FGuid& ArchivedSessionId : ArchivedSessionIds)
 	{
 		const FConcertServerSessionRepository& SessionRepository = GetSessionRepository(ArchivedSessionId);
-		if (SessionRepository.RepositoryId == Repository.RepositoryId)
+		if (SessionRepository.RepositoryId == RepositoryId)
 		{
-			DestroyArchivedSession(ArchivedSessionId, /*bDeleteSessionData*/false);
+			DestroyArchivedSession(ArchivedSessionId, /*bDeleteSessionData*/bDropped);
 		}
 	}
 
-	if (DefaultSessionRepository && DefaultSessionRepository->RepositoryId == Repository.RepositoryId)
+	if (DefaultSessionRepository && DefaultSessionRepository->RepositoryId == RepositoryId)
 	{
 		DefaultSessionRepository.Reset(); // Will not be able to create new sessions until a mounted repository is set as default.
-		UE_LOG(LogConcert, Warning, TEXT("Unmounted the default session repository. No session will be created until a mounted repository is set as default"));
+		UE_LOG(LogConcert, Warning, TEXT("Default repository %s unmounted. No session will be created until a mounted repository is set as default"), *RepositoryId.ToString());
 	}
-
-	// Mark the repository as unmounted.
-	Repository.bMounted = false;
-	Repository.ProcessId = 0;
-	MountedSessionRepositories.RemoveAll([Repository](const FConcertServerSessionRepository& CandidateRepository) { return Repository.RepositoryId == CandidateRepository.RepositoryId; });
-
-	// If the repository is being dropped, the file can go away if the repository is under a known root.
-	if (bDropped && !Repository.RepositoryRootDir.IsEmpty())
+	else
 	{
-		FString RepositoryDir = Repository.RepositoryRootDir / Repository.RepositoryId.ToString();
-		ConcertUtil::DeleteDirectoryTree(*RepositoryDir);
+		UE_LOG(LogConcert, Display, TEXT("Repository %s unmounted."), *RepositoryId.ToString())
 	}
 
-	// Update the repository database file.
-	FSystemWideCriticalSection ScopedSystemWideMutex(ConcertServerUtil::GetServerSystemMutexName());
-	FConcertServerSessionRepositoryDatabase SessionRepositoryDb;
-	ConcertServerUtil::LoadSessionRepositoryDatabase(Role, SessionRepositoryDb);
-	if (bDropped)
+	if (bDropped && !Repository.RepositoryRootDir.IsEmpty()) // When dropped, the repository can be deleted if it has the standard root structure.
 	{
-		SessionRepositoryDb.Repositories.RemoveAll([Repository](const FConcertServerSessionRepository& CandidateRepository) { return Repository.RepositoryId == CandidateRepository.RepositoryId; });
+		FString RepositoryDir = Repository.RepositoryRootDir / Repository.RepositoryId.ToString();
+		if (ConcertUtil::DeleteDirectoryTree(*RepositoryDir))
+		{
+			UE_LOG(LogConcert, Display, TEXT("Repository %s deleted."), *Repository.RepositoryId.ToString())
+		}
 	}
-	else if (FConcertServerSessionRepository* UnmountedRepo = SessionRepositoryDb.Repositories.FindByPredicate([Repository](const FConcertServerSessionRepository& CandidateRepository) { return Repository.RepositoryId == CandidateRepository.RepositoryId; }))
+
+	// Remove the repository from the of mounted repository list.
+	MountedSessionRepositories.RemoveAt(Index);
+
+	// Update the repository database file
 	{
-		UnmountedRepo->bMounted = false;
-		UnmountedRepo->ProcessId = 0;
+		FSystemWideCriticalSection ScopedSystemWideMutex(ConcertServerUtil::GetServerSystemMutexName());
+		FConcertServerSessionRepositoryDatabase SessionRepositoryDb;
+		ConcertServerUtil::LoadSessionRepositoryDatabase(Role, SessionRepositoryDb);
+		if (bDropped)
+		{
+			SessionRepositoryDb.Repositories.RemoveAll([&RepositoryId](const FConcertServerSessionRepository& RemoveCandidate) { return RepositoryId == RemoveCandidate.RepositoryId; });
+		}
+		else if (FConcertServerSessionRepository* UnmountedRepo = SessionRepositoryDb.Repositories.FindByPredicate([&RepositoryId](const FConcertServerSessionRepository& MatchRepository) { return RepositoryId == MatchRepository.RepositoryId; }))
+		{
+			UnmountedRepo->bMounted = false;
+			UnmountedRepo->ProcessId = 0;
+		}
+		ConcertServerUtil::SaveSessionRepositoryDatabase(Role, SessionRepositoryDb);
 	}
-	ConcertServerUtil::SaveSessionRepositoryDatabase(Role, SessionRepositoryDb);
 
 	return true;
 }
@@ -822,6 +860,15 @@ TFuture<FConcertAdmin_MountSessionRepositoryResponse> FConcertServer::HandleMoun
 		}
 	}
 
+	if (ResponseData.MountStatus == EConcertSessionRepositoryMountResponseCode::Mounted)
+	{
+		UE_LOG(LogConcert, Display, TEXT("User mounted repository %s"), *Message->RepositoryId.ToString());
+	}
+	else
+	{
+		UE_LOG(LogConcert, Display, TEXT("User failed to mount repository %s"), *Message->RepositoryId.ToString());
+	}
+
 	return FConcertAdmin_MountSessionRepositoryResponse::AsFuture(MoveTemp(ResponseData));
 }
 
@@ -861,49 +908,48 @@ TFuture<FConcertAdmin_DropSessionRepositoriesResponse> FConcertServer::HandleDro
 	const FConcertAdmin_DropSessionRepositoriesRequest* Message = Context.GetMessage<FConcertAdmin_DropSessionRepositoriesRequest>();
 	FConcertAdmin_DropSessionRepositoriesResponse ResponseData;
 
-	// Drop the repository currently loaded in-memory by this process.
+	// Drop the repository currently mounted by this process.
 	for (const FGuid& RepositoryId : Message->RepositoryIds)
 	{
-		if (FConcertServerSessionRepository* MountedRepo = MountedSessionRepositories.FindByPredicate([RepositoryId](const FConcertServerSessionRepository& Candidate) { return Candidate.RepositoryId == RepositoryId; }))
+		if (UnmountSessionRepository(RepositoryId, /*bDropped*/true))
 		{
-			if (UnmountSessionRepository(*MountedRepo, /*bDropped*/true))
-			{
-				ResponseData.DroppedRepositoryIds.Add(RepositoryId);
-			}
+			ResponseData.DroppedRepositoryIds.Add(RepositoryId);
 		}
 	}
 
-	// Drop the repository that aren't loaded by this process, but found in the global repository database.
+	// Drop the repository that aren't mounted, but found in the global repository database.
 	{
 		FSystemWideCriticalSection ScopedSystemWideMutex(ConcertServerUtil::GetServerSystemMutexName());
 		FConcertServerSessionRepositoryDatabase SessionRepositoryDb;
 		ConcertServerUtil::LoadSessionRepositoryDatabase(Role, SessionRepositoryDb);
-		int32 DroppedRepositoryNum = 0;
 
 		// Drop the repositories.
 		for (const FGuid& RepositoryId : Message->RepositoryIds)
 		{
-			DroppedRepositoryNum += SessionRepositoryDb.Repositories.RemoveAll([this, &RepositoryId, &ResponseData](FConcertServerSessionRepository& Repository)
+			int32 Index = SessionRepositoryDb.Repositories.IndexOfByPredicate([&RepositoryId](const FConcertServerSessionRepository& Repository) { return Repository.RepositoryId == RepositoryId; });
+			if (Index == INDEX_NONE)
+			{
+				ResponseData.DroppedRepositoryIds.Add(RepositoryId); // Not mapped in the DB -> successufully dropped.
+				continue;
+			}
+			
+			FConcertServerSessionRepository& Repository = SessionRepositoryDb.Repositories[Index];
+			if (!Repository.bMounted || !FPlatformProcess::IsApplicationRunning(Repository.ProcessId)) // Not mounted or mounted by a dead process.
 			{
-				bool bDropped = false;
-				if (Repository.RepositoryId == RepositoryId)
+				// Check if the server can delete the folder safely i.e. it has the standard structure managed by the server.
+				if (!Repository.RepositoryRootDir.IsEmpty())
 				{
-					if (!Repository.bMounted || !FPlatformProcess::IsApplicationRunning(Repository.ProcessId)) // Not mounted.
-					{
-						if (!Repository.RepositoryRootDir.IsEmpty()) // This is the standard repo form and the server know what should be deleted.
-						{
-							FString ReposDir = Repository.RepositoryRootDir / Repository.RepositoryId.ToString();
-							ConcertUtil::DeleteDirectoryTree(*ReposDir);
-						}
-						ResponseData.DroppedRepositoryIds.Add(RepositoryId);
-						return true;
-					}
+					FString ReposDir = Repository.RepositoryRootDir / Repository.RepositoryId.ToString();
+					ConcertUtil::DeleteDirectoryTree(*ReposDir);
 				}
-				return false;
-			});
+
+				// Unmap it.
+				SessionRepositoryDb.Repositories.RemoveAt(Index);
+				ResponseData.DroppedRepositoryIds.Add(RepositoryId);
+			}
 		}
 
-		if (DroppedRepositoryNum)
+		if (ResponseData.DroppedRepositoryIds.Num())
 		{
 			ConcertServerUtil::SaveSessionRepositoryDatabase(Role, SessionRepositoryDb);
 		}
diff --git a/Engine/Plugins/Developer/Concert/ConcertMain/Source/Concert/Private/ConcertServer.h b/Engine/Plugins/Developer/Concert/ConcertMain/Source/Concert/Private/ConcertServer.h
--- a/Engine/Plugins/Developer/Concert/ConcertMain/Source/Concert/Private/ConcertServer.h
+++ b/Engine/Plugins/Developer/Concert/ConcertMain/Source/Concert/Private/ConcertServer.h
@@ -54,7 +54,7 @@ private:
 	FString GetSessionWorkingDir(const FGuid& SessionId) const;
 
 	EConcertSessionRepositoryMountResponseCode MountSessionRepository(FConcertServerSessionRepository& Repository, bool bCreateIfNotExist, bool bCleanWorkingDir, bool bCleanExpiredSessions, bool bSearchByPaths);
-	bool UnmountSessionRepository(FConcertServerSessionRepository& InOutRepository, bool bDropped);
+	bool UnmountSessionRepository(const FGuid& RepositoryId, bool bDropped);
 
 	/**  */
 	void HandleDiscoverServersEvent(const FConcertMessageContext& Context);
