Fixed disaster recovery crashing if multiple instances were launched concurrently (because database files were locked)
  - Prevented multiple instances of disaster recovery services/multi-user server from opening the same files.

Added a new API to generate the recovery service server name available by both the client and the server.
  - Updated the command line used to launch the crash reporter client (embedding the recovery service).

diff --git a/Engine/Plugins/Developer/Concert/ConcertMain/Source/Concert/Private/ConcertServer.cpp b/Engine/Plugins/Developer/Concert/ConcertMain/Source/Concert/Private/ConcertServer.cpp
--- a/Engine/Plugins/Developer/Concert/ConcertMain/Source/Concert/Private/ConcertServer.cpp
+++ b/Engine/Plugins/Developer/Concert/ConcertMain/Source/Concert/Private/ConcertServer.cpp
@@ -9,10 +9,17 @@
 #include "ConcertLogGlobal.h"
 #include "IConcertServerEventSink.h"
 
+#include "Algo/AnyOf.h"
 #include "Misc/App.h"
 #include "Misc/Paths.h"
 
 #include "Runtime/Launch/Resources/Version.h"
+#include "ConcertServerInstanceInfo.h"
+#include "StructSerializer.h"
+#include "StructDeserializer.h"
+#include "Backends/JsonStructSerializerBackend.h"
+#include "Backends/JsonStructDeserializerBackend.h"
+#include "HAL/FileManager.h"
 
 #define LOCTEXT_NAMESPACE "ConcertServer"
 
@@ -31,6 +38,62 @@ FString GetArchiveName(const FString& SessionName, const FConcertSessionSettings
 	}
 }
 
+FString GetServerInstanceInfoPathname(const FString& Role)
+{
+	return FPaths::ProjectSavedDir() / Role / TEXT("InstanceInfo.json");
+}
+
+bool SaveServerInstanceInfo(const FString& Role, const FConcertServerInstanceInfo& InInstanceInfo)
+{
+	if (TUniquePtr<FArchive> FileWriter = TUniquePtr<FArchive>(IFileManager::Get().CreateFileWriter(*GetServerInstanceInfoPathname(Role))))
+	{
+		FJsonStructSerializerBackend Backend(*FileWriter, EStructSerializerBackendFlags::Default);
+		FStructSerializer::Serialize(InInstanceInfo, Backend);
+
+		FileWriter->Close();
+		return !FileWriter->IsError();
+	}
+
+	return false;
+}
+
+bool LoadServerInstanceInfo(const FString& Role, FConcertServerInstanceInfo& OutIntanceInfo)
+{
+	if (TUniquePtr<FArchive> FileReader = TUniquePtr<FArchive>(IFileManager::Get().CreateFileReader(*GetServerInstanceInfoPathname(Role))))
+	{
+		FJsonStructDeserializerBackend Backend(*FileReader);
+		FStructDeserializer::Deserialize(OutIntanceInfo, Backend);
+
+		FileReader->Close();
+		return !FileReader->IsError();
+	}
+
+	return false;
+}
+
+int32 RemoveInactiveServerInstances(FConcertServerInstanceInfo& InOutIntanceInfo)
+{
+	InOutIntanceInfo.Instances.RemoveAll([](const FConcertServerInstance& Instance)
+	{
+		if (FPlatformProcess::IsApplicationRunning(Instance.ProcessId))
+		{
+			// If the PID was reused by another application, remove it from the list.
+			return FPaths::GetPathLeaf(FPlatformProcess::GetApplicationName(Instance.ProcessId)) != FPaths::GetPathLeaf(FPlatformProcess::GetApplicationName(FPlatformProcess::GetCurrentProcessId()));
+		}
+		return true; // Not running, remove it.
+	});
+
+	return InOutIntanceInfo.Instances.Num(); // How many remains?
+}
+
+bool SharesDirectoriesWithOtherInstances(const FConcertServerPaths& Paths, FConcertServerInstanceInfo& InstanceInfo)
+{
+	return Algo::AnyOf(InstanceInfo.Instances, [&Paths](const FConcertServerInstance& Instance)
+	{
+		return Instance.ArchiveDirectory == Paths.GetSavedDir() || Instance.WorkingDirectory == Paths.GetWorkingDir();
+	});
+}
+
 }
 
 
@@ -128,21 +191,45 @@ void FConcertServer::Startup()
 		ServerAdminEndpoint->RegisterRequestHandler<FConcertAdmin_GetSessionClientsRequest, FConcertAdmin_GetSessionClientsResponse>(this, &FConcertServer::HandleGetSessionClientsRequest);
 		ServerAdminEndpoint->RegisterRequestHandler<FConcertAdmin_GetSessionActivitiesRequest, FConcertAdmin_GetSessionActivitiesResponse>(this, &FConcertServer::HandleGetSessionActivitiesRequest);
 
-		if (Settings->bCleanWorkingDir)
-		{
-			ConcertUtil::DeleteDirectoryTree(*Paths->GetWorkingDir(), *Paths->GetBaseWorkingDir());
-		}
-		else
+		// Concurrent server instances may fight to get ownership of the info file. Use a system wide mutex with an arbitrary name that will unlikely be found in other applications.
+		FSystemWideCriticalSection ScopedSystemWideMutex(TEXT("Unreal_ConcertServer_67822dAB"));
+
+		// Load the file containing the instance info.
+		FConcertServerInstanceInfo InstanceInfo;
+		ConcertServerUtil::LoadServerInstanceInfo(Role, InstanceInfo);
+
+		// Check if the server can scan/load/rotate the existing session files in the configured directories. (The existing session files are not sharable, they can only managed by one instance)
+		int32 ActiveInstanceCount = ConcertServerUtil::RemoveInactiveServerInstances(InstanceInfo);
+		bool bLoadExistingSessions = ActiveInstanceCount == 0 || !ConcertServerUtil::SharesDirectoriesWithOtherInstances(*Paths, InstanceInfo);
+		if (bLoadExistingSessions)
 		{
-			if (Settings->bAutoArchiveOnReboot)
+			if (Settings->bCleanWorkingDir)
 			{
-				// Migrate live sessions files (session is not restored yet) to its archive form and directory.
-				ArchiveOfflineSessions();
+				ConcertUtil::DeleteDirectoryTree(*Paths->GetWorkingDir(), *Paths->GetBaseWorkingDir());
 			}
+			else
+			{
+				if (Settings->bAutoArchiveOnReboot)
+				{
+					// Migrate live sessions files (session is not restored yet) to its archive form and directory.
+					ArchiveOfflineSessions();
+				}
 
-			// Build the list of archive/live sessions and rotate the list of archive to prevent having too many of them.
-			RecoverSessions();
+				// Build the list of archive/live sessions and rotate the list of archive to prevent having too many of them.
+				RecoverSessions();
+			}
 		}
+		// else -> This instance will be able to create/rename/delete new sessions but will not see/expose the currently existing ones to avoid file conflict/corruption.
+
+		// Add this instance to the instance info file.
+		FConcertServerInstance& ThisInstance = InstanceInfo.Instances.AddDefaulted_GetRef();
+		ThisInstance.ServerName = Settings->ServerName;
+		ThisInstance.WorkingDirectory = Paths->GetWorkingDir();
+		ThisInstance.ArchiveDirectory = Paths->GetSavedDir();
+		ThisInstance.ProcessId = FPlatformProcess::GetCurrentProcessId();
+		ThisInstance.bEnclaved = !bLoadExistingSessions; // Enclaved means this instance did not scan and load existing sessions to avoid conflict with another running instance(s). It was restricted to work with its own sessions only.
+
+		ConcertServerUtil::SaveServerInstanceInfo(Role, InstanceInfo);
 	}
 }
 
diff --git a/Engine/Plugins/Developer/Concert/ConcertMain/Source/Concert/Private/ConcertServerInstanceInfo.h b/Engine/Plugins/Developer/Concert/ConcertMain/Source/Concert/Private/ConcertServerInstanceInfo.h
new file mode 100644
--- /dev/null
+++ b/Engine/Plugins/Developer/Concert/ConcertMain/Source/Concert/Private/ConcertServerInstanceInfo.h
@@ -0,0 +1,39 @@
+// Copyright 1998-2019 Epic Games, Inc. All Rights Reserved.
+
+#pragma once
+
+#include "CoreMinimal.h"
+#include "ConcertServerInstanceInfo.generated.h"
+
+/** Keeps information about a concert server instance. */
+USTRUCT()
+struct FConcertServerInstance
+{
+	GENERATED_BODY()
+
+	UPROPERTY()
+	FString ServerName;
+
+	UPROPERTY()
+	FString WorkingDirectory;
+
+	UPROPERTY()
+	FString ArchiveDirectory;
+
+	UPROPERTY()
+	int32 ProcessId = 0;
+
+	/** True if the server instance was not allowed to scan/load existing sessions in the working/archived directories to avoid file conflicts with other instances. */
+	UPROPERTY()
+	bool bEnclaved = false;
+};
+
+/** Keeps information about all concert server instances. */
+USTRUCT()
+struct FConcertServerInstanceInfo
+{
+	GENERATED_BODY()
+
+	UPROPERTY()
+	TArray<FConcertServerInstance> Instances;
+};
diff --git a/Engine/Source/Runtime/Core/Private/GenericPlatform/GenericPlatformCrashContext.cpp b/Engine/Source/Runtime/Core/Private/GenericPlatform/GenericPlatformCrashContext.cpp
--- a/Engine/Source/Runtime/Core/Private/GenericPlatform/GenericPlatformCrashContext.cpp
+++ b/Engine/Source/Runtime/Core/Private/GenericPlatform/GenericPlatformCrashContext.cpp
@@ -1056,5 +1056,9 @@ FProgramCounterSymbolInfoEx::FProgramCounterSymbolInfoEx( FString InModuleName,
 	OffsetInModule( InOffsetInModule ),
 	ProgramCounter( InProgramCounter )
 {
+}
 
+FString RecoveryService::GetRecoveryServerName()
+{
+	return FString::Printf(TEXT("RecoverySvr_%d"), FPlatformProcess::GetCurrentProcessId());
 }
diff --git a/Engine/Source/Runtime/Core/Private/Windows/WindowsPlatformCrashContext.cpp b/Engine/Source/Runtime/Core/Private/Windows/WindowsPlatformCrashContext.cpp
--- a/Engine/Source/Runtime/Core/Private/Windows/WindowsPlatformCrashContext.cpp
+++ b/Engine/Source/Runtime/Core/Private/Windows/WindowsPlatformCrashContext.cpp
@@ -45,7 +45,7 @@
 #define CR_CLIENT_MAX_PATH_LEN 265
 #define CR_CLIENT_MAX_ARGS_LEN 256
 // Enable to use a non shipping version of the crash report client
-#define CR_USE_DEVELOPMENT_CLIENT 0
+#define CR_USE_DEVELOPMENT_CLIENT 1
 
 #pragma comment( lib, "version.lib" )
 #pragma comment( lib, "Shlwapi.lib" )
@@ -399,10 +399,8 @@ FProcHandle LaunchCrashReportClient(void** OutWritePipe, void** OutReadPipe)
 #if WITH_EDITOR // Disaster recovery is only enabled for the Editor. Start the server even if in -game, -server, commandlet, the client-side will not connect (its too soon here to query this executable config).
 	{
 		// Disaster recovery service command line.
-		const FString DisasterRecoveryServerName = FString::Printf(TEXT("RecoverySvr_%d"), FPlatformProcess::GetCurrentProcessId());
-
 		FString DisasterRecoveryServiceCommandLine;
-		DisasterRecoveryServiceCommandLine += FString::Printf(TEXT(" -ConcertServer=\"%s\""), *DisasterRecoveryServerName); // Must be in-sync with disaster recovery client module.
+		DisasterRecoveryServiceCommandLine += FString::Printf(TEXT(" -ConcertServer=\"%s\""), *RecoveryService::GetRecoveryServerName()); // Must be in-sync with disaster recovery client module.
 
 		FCString::Strncat(CrashReporterClientArgs, *DisasterRecoveryServiceCommandLine, CR_CLIENT_MAX_ARGS_LEN);
 	}
diff --git a/Engine/Source/Runtime/Core/Public/GenericPlatform/GenericPlatformCrashContext.h b/Engine/Source/Runtime/Core/Public/GenericPlatform/GenericPlatformCrashContext.h
--- a/Engine/Source/Runtime/Core/Public/GenericPlatform/GenericPlatformCrashContext.h
+++ b/Engine/Source/Runtime/Core/Public/GenericPlatform/GenericPlatformCrashContext.h
@@ -450,3 +450,9 @@ private:
 
 struct CORE_API FGenericMemoryWarningContext
 {};
+
+namespace RecoveryService
+{
+	/** Generates a name for the disaster recovery service embedded in the CrashReporterClient. */
+	CORE_API FString GetRecoveryServerName();
+}
