Fixed a off-by-one error when rotating the disaster recovery sessions that would delete the the oldest and the most recent session.

diff --git a/Engine/Plugins/Developer/Concert/ConcertMain/Source/Concert/Private/ConcertServer.cpp b/Engine/Plugins/Developer/Concert/ConcertMain/Source/Concert/Private/ConcertServer.cpp
--- a/Engine/Plugins/Developer/Concert/ConcertMain/Source/Concert/Private/ConcertServer.cpp
+++ b/Engine/Plugins/Developer/Concert/ConcertMain/Source/Concert/Private/ConcertServer.cpp
@@ -438,7 +438,7 @@ void FConcertServer::RecoverSessions()
 			// Keep the most recent sessions
 			TArray<FConcertSessionInfo> ArchivedSessionsToKeep;
 			{
-				const int32 FirstSortedSessionIndexToKeep = SortedSessions.Num() - Settings->NumSessionsToKeep - 1;
+				const int32 FirstSortedSessionIndexToKeep = SortedSessions.Num() - Settings->NumSessionsToKeep;
 				for (int32 SortedSessionIndex = FirstSortedSessionIndexToKeep; SortedSessionIndex < SortedSessions.Num(); ++SortedSessionIndex)
 				{
 					ArchivedSessionsToKeep.Add(ArchivedSessionInfos[SortedSessions[SortedSessionIndex].Key]);
