Fix plugin dependencies for CrashReportClientEditor

#rnx

diff --git a/Engine/Plugins/Developer/Concert/ConcertMain/ConcertMain.uplugin b/Engine/Plugins/Developer/Concert/ConcertMain/ConcertMain.uplugin
--- a/Engine/Plugins/Developer/Concert/ConcertMain/ConcertMain.uplugin
+++ b/Engine/Plugins/Developer/Concert/ConcertMain/ConcertMain.uplugin
@@ -15,19 +15,19 @@
 	"Hidden": true,
 	"IsBetaVersion": true,
 	"Installed": false,
-	"SupportedPrograms": [ "UnrealMultiUserServer", "UnrealDisasterRecoveryService", "CrashReportClient" ],
+	"SupportedPrograms": [ "UnrealMultiUserServer", "UnrealDisasterRecoveryService", "CrashReportClientEditor" ],
 	"Modules": [
 		{
 			"Name": "Concert",
 			"Type": "UncookedOnly",
 			"LoadingPhase": "PreDefault",
-			"WhitelistPrograms": [ "UnrealMultiUserServer", "UnrealDisasterRecoveryService", "CrashReportClient" ]
+			"WhitelistPrograms": [ "UnrealMultiUserServer", "UnrealDisasterRecoveryService", "CrashReportClientEditor" ]
 		},
 		{
 			"Name": "ConcertTransport",
 			"Type": "UncookedOnly",
 			"LoadingPhase": "PreDefault",
-			"WhitelistPrograms": [ "UnrealMultiUserServer", "UnrealDisasterRecoveryService", "CrashReportClient" ]
+			"WhitelistPrograms": [ "UnrealMultiUserServer", "UnrealDisasterRecoveryService", "CrashReportClientEditor" ]
 		}
 	]
 }
diff --git a/Engine/Plugins/Developer/Concert/ConcertSync/ConcertSyncCore/ConcertSyncCore.uplugin b/Engine/Plugins/Developer/Concert/ConcertSync/ConcertSyncCore/ConcertSyncCore.uplugin
--- a/Engine/Plugins/Developer/Concert/ConcertSync/ConcertSyncCore/ConcertSyncCore.uplugin
+++ b/Engine/Plugins/Developer/Concert/ConcertSync/ConcertSyncCore/ConcertSyncCore.uplugin
@@ -15,13 +15,13 @@
 	"Hidden": true,
 	"IsBetaVersion": true,
 	"Installed": false,
-	"SupportedPrograms": [ "UnrealMultiUserServer", "UnrealDisasterRecoveryService", "CrashReportClient" ],
+	"SupportedPrograms": [ "UnrealMultiUserServer", "UnrealDisasterRecoveryService", "CrashReportClientEditor" ],
 	"Modules": [
 		{
 			"Name": "ConcertSyncCore",
 			"Type": "UncookedOnly",
 			"LoadingPhase": "PreDefault",
-			"WhitelistPrograms": [ "UnrealMultiUserServer", "UnrealDisasterRecoveryService", "CrashReportClient" ]
+			"WhitelistPrograms": [ "UnrealMultiUserServer", "UnrealDisasterRecoveryService", "CrashReportClientEditor" ]
 		}
 	],
 	"Plugins": [
diff --git a/Engine/Plugins/Developer/Concert/ConcertSync/ConcertSyncServer/ConcertSyncServer.uplugin b/Engine/Plugins/Developer/Concert/ConcertSync/ConcertSyncServer/ConcertSyncServer.uplugin
--- a/Engine/Plugins/Developer/Concert/ConcertSync/ConcertSyncServer/ConcertSyncServer.uplugin
+++ b/Engine/Plugins/Developer/Concert/ConcertSync/ConcertSyncServer/ConcertSyncServer.uplugin
@@ -15,13 +15,13 @@
 	"Hidden": true,
 	"IsBetaVersion": true,
 	"Installed": false,
-	"SupportedPrograms": [ "UnrealMultiUserServer", "UnrealDisasterRecoveryService", "CrashReportClient" ],
+	"SupportedPrograms": [ "UnrealMultiUserServer", "UnrealDisasterRecoveryService", "CrashReportClientEditor" ],
 	"Modules": [
 		{
 			"Name": "ConcertSyncServer",
 			"Type": "UncookedOnly",
 			"LoadingPhase": "Default",
-			"WhitelistPrograms": [ "UnrealMultiUserServer", "UnrealDisasterRecoveryService", "CrashReportClient" ]
+			"WhitelistPrograms": [ "UnrealMultiUserServer", "UnrealDisasterRecoveryService", "CrashReportClientEditor" ]
 		}
 	],
 	"Plugins": [
diff --git a/Engine/Plugins/Messaging/UdpMessaging/UdpMessaging.uplugin b/Engine/Plugins/Messaging/UdpMessaging/UdpMessaging.uplugin
--- a/Engine/Plugins/Messaging/UdpMessaging/UdpMessaging.uplugin
+++ b/Engine/Plugins/Messaging/UdpMessaging/UdpMessaging.uplugin
@@ -20,7 +20,7 @@
 		"UnrealMultiUserServer",
 		"UnrealDisasterRecoveryService",
 		"UnrealLightmass",
-		"CrashReportClient"
+		"CrashReportClientEditor"
 	],
 	"Modules" :
 	[
@@ -37,7 +37,7 @@
 				"UnrealMultiUserServer",
 				"UnrealDisasterRecoveryService",
 				"UnrealLightmass",
-				"CrashReportClient"
+				"CrashReportClientEditor"
 			]
 		}
 	]
