- Defaulted the Concert debug log level to 'Error' to avoid spamming the log when a remote endpoint disconnect.

Edigrated from 11445411.

diff --git a/Engine/Plugins/Developer/Concert/ConcertMain/Source/ConcertTransport/Public/ConcertLogGlobal.h b/Engine/Plugins/Developer/Concert/ConcertMain/Source/ConcertTransport/Public/ConcertLogGlobal.h
--- a/Engine/Plugins/Developer/Concert/ConcertMain/Source/ConcertTransport/Public/ConcertLogGlobal.h
+++ b/Engine/Plugins/Developer/Concert/ConcertMain/Source/ConcertTransport/Public/ConcertLogGlobal.h
@@ -6,7 +6,7 @@
 
 // The default debug verbosity level can be set as a GlobalDefinitions in an executable Target.cs build file to adapt to this executable context.
 #ifndef UE_LOG_CONCERT_DEBUG_VERBOSITY_LEVEL
-#define UE_LOG_CONCERT_DEBUG_VERBOSITY_LEVEL Warning
+#define UE_LOG_CONCERT_DEBUG_VERBOSITY_LEVEL Error
 #endif
 
 CONCERTTRANSPORT_API DECLARE_LOG_CATEGORY_EXTERN(LogConcert, Log, All);
