- Fixed concert server to restore a session (by creating a new one) in the 'default repository' as expected rather than the creating it in the repository containing the archived session.

Edigrated 11252374 from Dev-VirtualProduction.

diff --git a/Engine/Plugins/Developer/Concert/ConcertMain/Source/Concert/Private/ConcertServer.cpp b/Engine/Plugins/Developer/Concert/ConcertMain/Source/Concert/Private/ConcertServer.cpp
--- a/Engine/Plugins/Developer/Concert/ConcertMain/Source/Concert/Private/ConcertServer.cpp
+++ b/Engine/Plugins/Developer/Concert/ConcertMain/Source/Concert/Private/ConcertServer.cpp
@@ -1439,7 +1439,7 @@ TSharedPtr<IConcertServerSession> FConcertServer::RestoreArchivedSession(const F
 	if (const FConcertSessionInfo* ArchivedSessionInfo = ArchivedSessions.Find(ArchivedSessionId))
 	{
 		// Find the archived session repository to restore the session in the same one.
-		const FConcertServerSessionRepository& SessionRepository = GetSessionRepository(ArchivedSessionId);
+		const FConcertServerSessionRepository& ArchivedSessionRepository = GetSessionRepository(ArchivedSessionId);
 
 		FString LiveSessionName = NewSessionInfo.SessionName;
 		if (LiveSessionName.IsEmpty())
@@ -1486,10 +1486,12 @@ TSharedPtr<IConcertServerSession> FConcertServer::RestoreArchivedSession(const F
 			}
 		}
 
-		if (EventSink->RestoreSession(*this, ArchivedSessionId, SessionRepository.GetSessionWorkingDir(LiveSessionInfo.SessionId), LiveSessionInfo, SessionFilter))
+		// Restore the session in the default repository (where new sessions should be created), unless it is unset.
+		const FConcertServerSessionRepository& RestoredSessionRepository = DefaultSessionRepository.IsSet() ? DefaultSessionRepository.GetValue() : ArchivedSessionRepository;
+		if (EventSink->RestoreSession(*this, ArchivedSessionId, RestoredSessionRepository.GetSessionWorkingDir(LiveSessionInfo.SessionId), LiveSessionInfo, SessionFilter))
 		{
 			UE_LOG(LogConcert, Display, TEXT("Archived session '%s' (%s) was restored as '%s' (%s)"), *ArchivedSessionInfo->SessionName, *ArchivedSessionInfo->SessionId.ToString(), *LiveSessionInfo.SessionName, *LiveSessionInfo.SessionId.ToString());
-			return CreateLiveSession(LiveSessionInfo, SessionRepository);
+			return CreateLiveSession(LiveSessionInfo, RestoredSessionRepository);
 		}
 	}
 
