Modified Multi-User session connection text

diff --git a/Engine/Plugins/Developer/Concert/ConcertMain/Source/Concert/Private/ConcertClient.cpp b/Engine/Plugins/Developer/Concert/ConcertMain/Source/Concert/Private/ConcertClient.cpp
--- a/Engine/Plugins/Developer/Concert/ConcertMain/Source/Concert/Private/ConcertClient.cpp
+++ b/Engine/Plugins/Developer/Concert/ConcertMain/Source/Concert/Private/ConcertClient.cpp
@@ -79,7 +79,7 @@ public:
 private:
 	FText GetAutoConnectionCanceledMessage() const
 	{
-		return FText::Format(LOCTEXT("AutoJoinSessionCanceled", "Auto-Connection to Session '{0}' Canceled."), FText::AsCultureInvariant(Settings->DefaultSessionName));
+		return FText::Format(LOCTEXT("AutoJoinSessionCanceled", "Connection to Session '{0}' Canceled."), FText::AsCultureInvariant(Settings->DefaultSessionName));
 	}
 
 	TUniquePtr<FAsyncTaskNotification> MakeAutoConnectNotification()
