- Added details to the error message when a Multi-User session fails to create.
  - Added logs in the server to provide more insight when the default session repository fails to mount (usually because a concurrent instance is locking it).
  - Retry to mount the default repository (if configured) if it failed to mount at start up before failing to create a session.

diff --git a/Engine/Plugins/Developer/Concert/ConcertMain/Source/Concert/Private/ConcertServer.cpp b/Engine/Plugins/Developer/Concert/ConcertMain/Source/Concert/Private/ConcertServer.cpp
--- a/Engine/Plugins/Developer/Concert/ConcertMain/Source/Concert/Private/ConcertServer.cpp
+++ b/Engine/Plugins/Developer/Concert/ConcertMain/Source/Concert/Private/ConcertServer.cpp
@@ -81,6 +81,7 @@ bool LoadSessionRepositoryDatabase(const FString& Role, FConcertServerSessionRep
 
 FConcertServer::FConcertServer(const FString& InRole, const FConcertSessionFilter& InAutoArchiveSessionFilter, IConcertServerEventSink* InEventSink, const TSharedPtr<IConcertEndpointProvider>& InEndpointProvider)
 	: Role(InRole)
+	, DefaultSessionRepositoryStatus(LOCTEXT("SessionRepository_NotConfigured", "Repository not configured."))
 	, AutoArchiveSessionFilter(InAutoArchiveSessionFilter)
 	, EventSink(InEventSink)
 	, EndpointProvider(InEndpointProvider)
@@ -231,26 +232,8 @@ void FConcertServer::Startup()
 			}
 		}
 
-		// If the server was configured to use a custom working/archive dir, create a corresponding repository and try to mount it.
-		if (!Settings->WorkingDir.IsEmpty() || !Settings->ArchiveDir.IsEmpty())
-		{
-			FConcertServerSessionRepository Repository(Role, FGuid::NewGuid(), Settings->WorkingDir, Settings->ArchiveDir);
-			if (MountSessionRepository(Repository, /*bCreateIfNotExist*/true, Settings->bCleanWorkingDir, /*bCleanupExpiredSession*/true, /*bSearchByPath*/true) == EConcertSessionRepositoryMountResponseCode::Mounted)
-			{
-				check(Repository.bMounted && Repository.ProcessId == FPlatformProcess::GetCurrentProcessId());
-				DefaultSessionRepository.Emplace(Repository);
-			}
-		}
-		// If the server was configured to mount a default server managed repository.
-		else if (Settings->bMountDefaultSessionRepository)
-		{
-			FConcertServerSessionRepository Repository(GetSessionRepositoriesRootDir(), FGuid()); // Invalid GUID is used for the default server repository.
-			if (MountSessionRepository(Repository, /*bCreateIfNotExist*/true, Settings->bCleanWorkingDir, /*bCleanupExpiredSession*/true, /*bSearchByPath*/false) == EConcertSessionRepositoryMountResponseCode::Mounted)
-			{
-				check(Repository.bMounted && Repository.ProcessId == FPlatformProcess::GetCurrentProcessId());
-				DefaultSessionRepository.Emplace(Repository);
-			}
-		}
+		// Try to mount the default session repository configured (if one is configured) to lock the non-sharable session files away from concurrent processes.
+		MountDefaultSessionRepository(Settings.Get());
 	}
 }
 
@@ -400,9 +383,12 @@ TSharedPtr<IConcertServerSession> FConcertServer::CreateSession(const FConcertSe
 		return nullptr;
 	}
 
-	if (!DefaultSessionRepository)
+	// If the default session repository is not set, check if one is configured and try to mount it. This may be a time-costly operation. This is to addresses the case where a user
+	// has/had two concurrent Multi-User servers using the same sessions directories without noticing and fail to create a session on the newest server instance because the folder is/was
+	// locked by the older instance when the new one started.
+	if (!DefaultSessionRepository && !MountDefaultSessionRepository(Settings.Get()))
 	{
-		OutFailureReason = FText::Format(LOCTEXT("Error_CreateSession_NoRepository", "Session '{0}' could not be created, no session repository was mounted to store it"), FText::AsCultureInvariant(SessionInfo.SessionName));
+		OutFailureReason = FText::Format(LOCTEXT("Error_CreateSession_NoRepository", "Session '{0}' could not be created. The default repository used to store sessions files is not mounted. Reason: {1}"), FText::AsCultureInvariant(SessionInfo.SessionName), DefaultSessionRepositoryStatus);
 		UE_LOG(LogConcert, Error, TEXT("An attempt to create a session with name '%s' was made, but the server did not have any repository mounted to store it!"), *SessionInfo.SessionName);
 		return nullptr;
 	}
@@ -698,8 +684,12 @@ FString FConcertServer::GetSessionWorkingDir(const FGuid& SessionId) const
 	return GetSessionRepository(SessionId).GetSessionWorkingDir(SessionId);
 }
 
-EConcertSessionRepositoryMountResponseCode FConcertServer::MountSessionRepository(FConcertServerSessionRepository& Repository, bool bCreateIfNotExist, bool bCleanWorkingDir, bool bCleanExpiredSessions, bool bSearchByPaths)
+EConcertSessionRepositoryMountResponseCode FConcertServer::MountSessionRepository(FConcertServerSessionRepository Repository, bool bCreateIfNotExist, bool bCleanWorkingDir, bool bCleanExpiredSessions, bool bSearchByPaths, bool bAsDefault)
 {
+	EConcertSessionRepositoryMountResponseCode MountStatus = EConcertSessionRepositoryMountResponseCode::Mounted;
+	FText MountStatusText = LOCTEXT("SessionRepository_Mounted", "Repository mounted.");
+	bool bAlreadyMountedByThisProcess = false;
+
 	// Exclusive access scope to the session repository db.
 	{
 		// Load the file containing the instance/repository info.
@@ -725,11 +715,14 @@ EConcertSessionRepositoryMountResponseCode FConcertServer::MountSessionRepositor
 			else if (ExistingRepository->ProcessId == FPlatformProcess::GetCurrentProcessId() &&
 				MountedSessionRepositories.ContainsByPredicate([&Repository](const FConcertServerSessionRepository& MatchCandidate){ return MatchCandidate.RepositoryId == Repository.RepositoryId; })) // Already mounted by this process?
 			{
-				return EConcertSessionRepositoryMountResponseCode::Mounted;
+				UE_LOG(LogConcert, Display, TEXT("Remounted repository %s. The repository is already mounted by this process."), *Repository.RepositoryId.ToString());
+				bAlreadyMountedByThisProcess = true; // Already mounted by this process, don't process the session files again.
 			}
 			else
 			{
-				return EConcertSessionRepositoryMountResponseCode::AlreadyMounted; // Already mounted by another process, cannot mount it.
+				UE_LOG(LogConcert, Warning, TEXT("Failed to mount repository %s. The repository is already mounted by another process."), *Repository.RepositoryId.ToString());
+				MountStatus = EConcertSessionRepositoryMountResponseCode::AlreadyMounted; // Already mounted by another process, cannot mount it, the files are not shareable.
+				MountStatusText = LOCTEXT("SessionRepository_AlreadyMounted", "Repository locked by another process.");
 			}
 		}
 		else if (bCreateIfNotExist)
@@ -742,25 +735,47 @@ EConcertSessionRepositoryMountResponseCode FConcertServer::MountSessionRepositor
 		}
 		else
 		{
-			return EConcertSessionRepositoryMountResponseCode::NotFound;
+			UE_LOG(LogConcert, Warning, TEXT("Failed to mount repository %s. The repository was not found."), *Repository.RepositoryId.ToString());
+			MountStatus = EConcertSessionRepositoryMountResponseCode::NotFound;
+			MountStatusText = LOCTEXT("SessionRepository_NotFound", "Repository not found.");
 		}
 	}
 
-	// Process the sessions in the repository.
-	if (bCleanWorkingDir)
+	// Should the mounted repository be used as default?
+	if (bAsDefault)
 	{
-		ConcertUtil::DeleteDirectoryTree(*Repository.WorkingDir);
+		if (MountStatus == EConcertSessionRepositoryMountResponseCode::Mounted)
+		{
+			DefaultSessionRepository = Repository;
+			UE_LOG(LogConcert, Display, TEXT("Default session repository %s set successfully."), *Repository.RepositoryId.ToString());
+		}
+		else
+		{
+			DefaultSessionRepository.Reset();
+			UE_LOG(LogConcert, Warning, TEXT("Default session repository %s failed to mount."), *Repository.RepositoryId.ToString());
+		}
+		DefaultSessionRepositoryStatus = MountStatusText;
 	}
-	else if (Settings->bAutoArchiveOnReboot) // Honor the auto-archive settings when mounting a new repository.
+
+	// Should the sessions in the repository processed?
+	if (MountStatus == EConcertSessionRepositoryMountResponseCode::Mounted && !bAlreadyMountedByThisProcess)
 	{
-		// Migrate live sessions files (session is not restored yet) to its archive form and directory.
-		ArchiveOfflineSessions(Repository);
-	}
+		// Process the sessions in the repository.
+		if (bCleanWorkingDir)
+		{
+			ConcertUtil::DeleteDirectoryTree(*Repository.WorkingDir);
+		}
+		else if (Settings->bAutoArchiveOnReboot) // Honor the auto-archive settings when mounting a new repository.
+		{
+			// Migrate live sessions files (session is not restored yet) to its archive form and directory.
+			ArchiveOfflineSessions(Repository);
+		}
 
-	// Reload the archived/live sessions and possibly rotate the list of archives to prevent having too many of them.
-	RecoverSessions(Repository, bCleanExpiredSessions);
+		// Reload the archived/live sessions and possibly rotate the list of archives to prevent having too many of them.
+		RecoverSessions(Repository, bCleanExpiredSessions);
+	}
 
-	return EConcertSessionRepositoryMountResponseCode::Mounted;
+	return MountStatus;
 }
 
 bool FConcertServer::UnmountSessionRepository(const FGuid& RepositoryId, bool bDropped)
@@ -803,6 +818,7 @@ bool FConcertServer::UnmountSessionRepository(const FGuid& RepositoryId, bool bD
 	if (DefaultSessionRepository && DefaultSessionRepository->RepositoryId == RepositoryId)
 	{
 		DefaultSessionRepository.Reset(); // Will not be able to create new sessions until a mounted repository is set as default.
+		DefaultSessionRepositoryStatus = LOCTEXT("SessionRepository_Unmounted", "Repository unmounted.");
 		UE_LOG(LogConcert, Warning, TEXT("Default repository %s unmounted. No session will be created until a mounted repository is set as default"), *RepositoryId.ToString());
 	}
 	else
@@ -842,6 +858,29 @@ bool FConcertServer::UnmountSessionRepository(const FGuid& RepositoryId, bool bD
 	return true;
 }
 
+bool FConcertServer::MountDefaultSessionRepository(const UConcertServerConfig* ServerConfig)
+{
+	if (DefaultSessionRepository)
+	{
+		return true; // A default session repository is already mounted.
+	}
+
+	// If the server was configured to use a custom working/archive dir, create a corresponding repository and try to mount it.
+	if (!ServerConfig->WorkingDir.IsEmpty() || !ServerConfig->ArchiveDir.IsEmpty())
+	{
+		FConcertServerSessionRepository Repository(Role, FGuid::NewGuid(), ServerConfig->WorkingDir, ServerConfig->ArchiveDir);
+		return MountSessionRepository(MoveTemp(Repository), /*bCreateIfNotExist*/true, ServerConfig->bCleanWorkingDir, /*bCleanupExpiredSession*/true, /*bSearchByPath*/true, /*bAsDefault*/true) == EConcertSessionRepositoryMountResponseCode::Mounted;
+	}
+	// If the server was configured to mount a default server managed repository.
+	else if (ServerConfig->bMountDefaultSessionRepository)
+	{
+		FConcertServerSessionRepository Repository(GetSessionRepositoriesRootDir(), FGuid()); // Invalid GUID is used for the default server repository.
+		return MountSessionRepository(MoveTemp(Repository), /*bCreateIfNotExist*/true, ServerConfig->bCleanWorkingDir, /*bCleanupExpiredSession*/true, /*bSearchByPath*/false, /*bAsDefault*/true) == EConcertSessionRepositoryMountResponseCode::Mounted;
+	}
+
+	return false; // No session repository was mounted as default.
+}
+
 void FConcertServer::HandleDiscoverServersEvent(const FConcertMessageContext& Context)
 {
 	const FConcertAdmin_DiscoverServersEvent* Message = Context.GetMessage<FConcertAdmin_DiscoverServersEvent>();
@@ -867,35 +906,12 @@ TFuture<FConcertAdmin_MountSessionRepositoryResponse> FConcertServer::HandleMoun
 	if (Message->RepositoryRootDir.IsEmpty()) // Use the server configured repository root dir?
 	{
 		FConcertServerSessionRepository Repository(GetSessionRepositoriesRootDir(), Message->RepositoryId);
-		ResponseData.MountStatus = MountSessionRepository(Repository, Message->bCreateIfNotExist, /*bCleanWorkingDir*/false, /*bCleanExpiredSessions*/false, /*bSearchByPaths*/false);
+		ResponseData.MountStatus = MountSessionRepository(MoveTemp(Repository), Message->bCreateIfNotExist, /*bCleanWorkingDir*/false, /*bCleanExpiredSessions*/false, /*bSearchByPaths*/false, Message->bAsServerDefault);
 	}
 	else // Use the client supplied repository root dir.
 	{
 		FConcertServerSessionRepository Repository(Message->RepositoryRootDir, Message->RepositoryId);
-		ResponseData.MountStatus = MountSessionRepository(Repository, Message->bCreateIfNotExist, /*bCleanWorkingDir*/false, /*bCleanExpiredSessions*/false, /*bSearchByPaths*/false);
-	}
-
-	// Should the mounted repository be used as default?
-	if (Message->bAsServerDefault && ResponseData.MountStatus == EConcertSessionRepositoryMountResponseCode::Mounted)
-	{
-		const FConcertServerSessionRepository* DefaultRepository = MountedSessionRepositories.FindByPredicate([RepositoryId = Message->RepositoryId](const FConcertServerSessionRepository& Repository)
-		{
-			return Repository.RepositoryId == RepositoryId;
-		});
-
-		if (DefaultRepository)
-		{
-			DefaultSessionRepository = *DefaultRepository;
-		}
-	}
-
-	if (ResponseData.MountStatus == EConcertSessionRepositoryMountResponseCode::Mounted)
-	{
-		UE_LOG(LogConcert, Display, TEXT("User mounted repository %s"), *Message->RepositoryId.ToString());
-	}
-	else
-	{
-		UE_LOG(LogConcert, Display, TEXT("User failed to mount repository %s"), *Message->RepositoryId.ToString());
+		ResponseData.MountStatus = MountSessionRepository(MoveTemp(Repository), Message->bCreateIfNotExist, /*bCleanWorkingDir*/false, /*bCleanExpiredSessions*/false, /*bSearchByPaths*/false, Message->bAsServerDefault);
 	}
 
 	return FConcertAdmin_MountSessionRepositoryResponse::AsFuture(MoveTemp(ResponseData));
diff --git a/Engine/Plugins/Developer/Concert/ConcertMain/Source/Concert/Private/ConcertServer.h b/Engine/Plugins/Developer/Concert/ConcertMain/Source/Concert/Private/ConcertServer.h
--- a/Engine/Plugins/Developer/Concert/ConcertMain/Source/Concert/Private/ConcertServer.h
+++ b/Engine/Plugins/Developer/Concert/ConcertMain/Source/Concert/Private/ConcertServer.h
@@ -54,8 +54,9 @@ private:
 	FString GetSessionSavedDir(const FGuid& SessionId) const;
 	FString GetSessionWorkingDir(const FGuid& SessionId) const;
 
-	EConcertSessionRepositoryMountResponseCode MountSessionRepository(FConcertServerSessionRepository& Repository, bool bCreateIfNotExist, bool bCleanWorkingDir, bool bCleanExpiredSessions, bool bSearchByPaths);
+	EConcertSessionRepositoryMountResponseCode MountSessionRepository(FConcertServerSessionRepository Repository, bool bCreateIfNotExist, bool bCleanWorkingDir, bool bCleanExpiredSessions, bool bSearchByPaths, bool bAsDefault);
 	bool UnmountSessionRepository(const FGuid& RepositoryId, bool bDropped);
+	bool MountDefaultSessionRepository(const UConcertServerConfig* ServerConfig);
 
 	/**  */
 	void HandleDiscoverServersEvent(const FConcertMessageContext& Context);
@@ -149,9 +150,12 @@ private:
 	/** All session repositories used by this server. */
 	TArray<FConcertServerSessionRepository> MountedSessionRepositories;
 
-	/** The server default repository created from the server settings. */
+	/** The default directory(ies) scanned to reload existing sessions or used to store newly created ones. */
 	TOptional<FConcertServerSessionRepository> DefaultSessionRepository;
 
+	/** The status of the default session repository (essentially to know why it not mounted). */
+	FText DefaultSessionRepositoryStatus;
+
 	/** The session filter to apply when auto-archiving sessions */
 	FConcertSessionFilter AutoArchiveSessionFilter;
 
