Clean-up UE4 names in Concert

diff --git a/Engine/Plugins/Developer/Concert/ConcertMain/Source/Concert/Private/ConcertVersion.cpp b/Engine/Plugins/Developer/Concert/ConcertMain/Source/Concert/Private/ConcertVersion.cpp
--- a/Engine/Plugins/Developer/Concert/ConcertMain/Source/Concert/Private/ConcertVersion.cpp
+++ b/Engine/Plugins/Developer/Concert/ConcertMain/Source/Concert/Private/ConcertVersion.cpp
@@ -49,14 +49,14 @@ bool ValidateVersion(const int32 InCurrent, const int32 InOther, const FText InV
 
 void FConcertFileVersionInfo::Initialize()
 {
-	FileVersionUE4 = GPackageFileUEVersion;
-	FileVersionLicenseeUE4 = GPackageFileLicenseeUEVersion;
+	FileVersion = GPackageFileUEVersion;
+	FileVersionLicensee = GPackageFileLicenseeUEVersion;
 }
 
 bool FConcertFileVersionInfo::Validate(const FConcertFileVersionInfo& InOther, const EConcertVersionValidationMode InValidationMode, FText* OutFailureReason) const
 {
-	return ConcertVersionUtil::ValidateVersion(FileVersionUE4, InOther.FileVersionUE4, LOCTEXT("PackageVersionName", "Package Version"), InValidationMode, OutFailureReason)
-		&& ConcertVersionUtil::ValidateVersion(FileVersionLicenseeUE4, InOther.FileVersionLicenseeUE4, LOCTEXT("LicenseePackageVersionName", "Licensee Package Version"), InValidationMode, OutFailureReason);
+	return ConcertVersionUtil::ValidateVersion(FileVersion, InOther.FileVersion, LOCTEXT("PackageVersionName", "Package Version"), InValidationMode, OutFailureReason)
+		&& ConcertVersionUtil::ValidateVersion(FileVersionLicensee, InOther.FileVersionLicensee, LOCTEXT("LicenseePackageVersionName", "Licensee Package Version"), InValidationMode, OutFailureReason);
 }
 
 
diff --git a/Engine/Plugins/Developer/Concert/ConcertMain/Source/Concert/Public/ConcertVersion.h b/Engine/Plugins/Developer/Concert/ConcertMain/Source/Concert/Public/ConcertVersion.h
--- a/Engine/Plugins/Developer/Concert/ConcertMain/Source/Concert/Public/ConcertVersion.h
+++ b/Engine/Plugins/Developer/Concert/ConcertMain/Source/Concert/Public/ConcertVersion.h
@@ -29,13 +29,13 @@ struct FConcertFileVersionInfo
 	/** Validate this version info against another */
 	CONCERT_API bool Validate(const FConcertFileVersionInfo& InOther, const EConcertVersionValidationMode InValidationMode, FText* OutFailureReason = nullptr) const;
 
-	/* UE4 file version */
+	/* File version */
 	UPROPERTY()
-	int32 FileVersionUE4 = 0;
+	int32 FileVersion = 0;
 
 	/* Licensee file version */
 	UPROPERTY()
-	int32 FileVersionLicenseeUE4 = 0;
+	int32 FileVersionLicensee = 0;
 };
 
 /** Holds engine version information */
diff --git a/Engine/Plugins/Developer/Concert/ConcertSync/ConcertSyncClient/Source/ConcertSyncClient/Private/ConcertSandboxPlatformFile.h b/Engine/Plugins/Developer/Concert/ConcertSync/ConcertSyncClient/Source/ConcertSyncClient/Private/ConcertSandboxPlatformFile.h
--- a/Engine/Plugins/Developer/Concert/ConcertSync/ConcertSyncClient/Source/ConcertSyncClient/Private/ConcertSandboxPlatformFile.h
+++ b/Engine/Plugins/Developer/Concert/ConcertSync/ConcertSyncClient/Source/ConcertSyncClient/Private/ConcertSandboxPlatformFile.h
@@ -216,10 +216,10 @@ private:
 		FDelegateHandle OnDirectoryChangedHandle;
 	};
 
-	/** Callback when a new content path is mounted in UE4 */
+	/** Callback when a new content path is mounted */
 	void OnContentPathMounted(const FString& InAssetPath, const FString& InFilesystemPath);
 
-	/** Callback when an existing content path is unmounted in UE4 */
+	/** Callback when an existing content path is unmounted */
 	void OnContentPathDismounted(const FString& InAssetPath, const FString& InFilesystemPath);
 
 	/** Register a mount path from a source content path */
diff --git a/Engine/Plugins/Developer/Concert/ConcertSync/ConcertSyncCore/Source/ConcertSyncCore/Private/ConcertSyncArchives.cpp b/Engine/Plugins/Developer/Concert/ConcertSync/ConcertSyncCore/Source/ConcertSyncCore/Private/ConcertSyncArchives.cpp
--- a/Engine/Plugins/Developer/Concert/ConcertSync/ConcertSyncCore/Source/ConcertSyncCore/Private/ConcertSyncArchives.cpp
+++ b/Engine/Plugins/Developer/Concert/ConcertSync/ConcertSyncCore/Source/ConcertSyncCore/Private/ConcertSyncArchives.cpp
@@ -180,8 +180,8 @@ FConcertSyncObjectReader::FConcertSyncObjectReader(const FConcertLocalIdentifier
 
 	if (InVersionInfo)
 	{
-		SetUE4Ver(InVersionInfo->FileVersion.FileVersionUE4);
-		SetLicenseeUE4Ver(InVersionInfo->FileVersion.FileVersionLicenseeUE4);
+		SetUE4Ver(InVersionInfo->FileVersion.FileVersion);
+		SetLicenseeUE4Ver(InVersionInfo->FileVersion.FileVersionLicensee);
 		SetEngineVer(FEngineVersionBase(InVersionInfo->EngineVersion.Major, InVersionInfo->EngineVersion.Minor, InVersionInfo->EngineVersion.Patch, InVersionInfo->EngineVersion.Changelist));
 
 		FCustomVersionContainer EngineCustomVersions;
