Messaging: replace more instances of message construction using operator new with FMessageEndpoint::MakeMessage()

diff --git a/Engine/Plugins/Animation/LiveLink/Source/LiveLink/Private/LiveLinkMessageBusSource.cpp b/Engine/Plugins/Animation/LiveLink/Source/LiveLink/Private/LiveLinkMessageBusSource.cpp
--- a/Engine/Plugins/Animation/LiveLink/Source/LiveLink/Private/LiveLinkMessageBusSource.cpp
+++ b/Engine/Plugins/Animation/LiveLink/Source/LiveLink/Private/LiveLinkMessageBusSource.cpp
@@ -252,7 +252,7 @@ FORCEINLINE void FLiveLinkMessageBusSource::UpdateConnectionLastActive()
 
 void FLiveLinkMessageBusSource::SendConnectMessage()
 {
-	FLiveLinkConnectMessage* ConnectMessage = new FLiveLinkConnectMessage();
+	FLiveLinkConnectMessage* ConnectMessage = FMessageEndpoint::MakeMessage<FLiveLinkConnectMessage>();
 	ConnectMessage->LiveLinkVersion = ILiveLinkClient::LIVELINK_VERSION;
 	MessageEndpoint->Send(ConnectMessage, ConnectionAddress);
 	FLiveLinkHeartbeatEmitter& HeartbeatEmitter = ILiveLinkModule::Get().GetHeartbeatEmitter();
diff --git a/Engine/Plugins/Developer/Concert/ConcertMain/Source/ConcertTransport/ConcertTransport.Build.cs b/Engine/Plugins/Developer/Concert/ConcertMain/Source/ConcertTransport/ConcertTransport.Build.cs
--- a/Engine/Plugins/Developer/Concert/ConcertMain/Source/ConcertTransport/ConcertTransport.Build.cs
+++ b/Engine/Plugins/Developer/Concert/ConcertMain/Source/ConcertTransport/ConcertTransport.Build.cs
@@ -11,16 +11,10 @@ namespace UnrealBuildTool.Rules
 				{
 					"Core",
 					"CoreUObject",
+					"MessagingCommon",
 					"TraceLog",
                 }
 			);
-
-            PrivateIncludePathModuleNames.AddRange(
-				new string[] {
-					"MessagingCommon",
-				}
-			);
-
         }
     }
 }
diff --git a/Engine/Plugins/Developer/Concert/ConcertMain/Source/ConcertTransport/Private/ConcertLocalEndpoint.cpp b/Engine/Plugins/Developer/Concert/ConcertMain/Source/ConcertTransport/Private/ConcertLocalEndpoint.cpp
--- a/Engine/Plugins/Developer/Concert/ConcertMain/Source/ConcertTransport/Private/ConcertLocalEndpoint.cpp
+++ b/Engine/Plugins/Developer/Concert/ConcertMain/Source/ConcertTransport/Private/ConcertLocalEndpoint.cpp
@@ -372,7 +372,7 @@ void FConcertLocalEndpoint::SendAck(const FGuid& AcknowledgmentToSend, const FCo
 		return;
 	}
 
-	FConcertAckData* Ack = new FConcertAckData();
+	FConcertAckData* Ack = FMessageEndpoint::MakeMessage<FConcertAckData>();
 	Ack->ConcertEndpointId = EndpointContext.EndpointId;
 	Ack->MessageId = FGuid::NewGuid();
 	Ack->AckSendTimeTicks = UtcNow.GetTicks();
@@ -400,7 +400,7 @@ void FConcertLocalEndpoint::SendEndpointClosed(const FConcertRemoteEndpointRef&
 		return;
 	}
 
-	FConcertEndpointClosedData* EndpointClosed = new FConcertEndpointClosedData();
+	FConcertEndpointClosedData* EndpointClosed = FMessageEndpoint::MakeMessage<FConcertEndpointClosedData>();
 	EndpointClosed->ConcertEndpointId = EndpointContext.EndpointId;
 	EndpointClosed->MessageId = FGuid::NewGuid();
 
@@ -569,7 +569,7 @@ void FConcertLocalEndpoint::ProcessEndpointDiscovery(const FConcertMessageContex
 		// Negotiate a reliable channel
 		if (MessageEndpoint.IsValid())
 		{
-			FConcertReliableHandshakeData* InitialHandshake = new FConcertReliableHandshakeData();
+			FConcertReliableHandshakeData* InitialHandshake = FMessageEndpoint::MakeMessage<FConcertReliableHandshakeData>();
 			InitialHandshake->ConcertEndpointId = EndpointContext.EndpointId;
 			InitialHandshake->MessageId = FGuid::NewGuid();
 			InitialHandshake->EndpointTimeoutTick = FTimespan(0, 0, Settings.RemoteEndpointTimeoutSeconds).GetTicks();
@@ -603,29 +603,33 @@ void FConcertLocalEndpoint::ProcessReliableHandshake(const FConcertMessageContex
 
 	Logger.LogReceiveReliableHandshake(*Message, EndpointContext.EndpointId, ConcertContext.UtcNow);
 
-	TUniquePtr<FConcertReliableHandshakeData> HandshakeResponse = MakeUnique<FConcertReliableHandshakeData>();
+	FConcertReliableHandshakeData* HandshakeResponse = FMessageEndpoint::MakeMessage<FConcertReliableHandshakeData>();
 	HandshakeResponse->ConcertEndpointId = EndpointContext.EndpointId;
 	HandshakeResponse->MessageId = FGuid::NewGuid();
 	HandshakeResponse->EndpointTimeoutTick = FTimespan(0, 0, Settings.RemoteEndpointTimeoutSeconds).GetTicks();
-	if (RemoteEndpoint->HandleReliableHandshake(*Message, *HandshakeResponse) && MessageEndpoint.IsValid())
+	if (!RemoteEndpoint->HandleReliableHandshake(*Message, *HandshakeResponse) || !MessageEndpoint.IsValid())
 	{
-		Logger.LogSendReliableHandshake(*HandshakeResponse, Message->ConcertEndpointId, ConcertContext.UtcNow);
+		HandshakeResponse->~FConcertReliableHandshakeData();
+		FMemory::Free(HandshakeResponse);
+		return;
+	}
 
-		// Update the last sent message time to this endpoint
-		RemoteEndpoint->SetLastSentMessageTime(ConcertContext.UtcNow);
+	Logger.LogSendReliableHandshake(*HandshakeResponse, Message->ConcertEndpointId, ConcertContext.UtcNow);
 
-		MessageEndpoint->Send(
-			HandshakeResponse.Release(), // Should be deleted by MessageBus
-			EMessageFlags::Reliable,
-			nullptr, // No Attachment
-			TArrayBuilder<FMessageAddress>().Add(RemoteEndpoint->GetAddress()),
-			FTimespan::Zero(), // No Delay
-			FDateTime::MaxValue() // No Expiration
-		);
+	// Update the last sent message time to this endpoint
+	RemoteEndpoint->SetLastSentMessageTime(ConcertContext.UtcNow);
 
-		// (Re)send any pending reliable messages
-		SendPendingMessages(RemoteEndpoint.ToSharedRef(), ConcertContext.UtcNow);
-	}
+	MessageEndpoint->Send(
+		HandshakeResponse, // Should be deleted by MessageBus
+		EMessageFlags::Reliable,
+		nullptr, // No Attachment
+		TArrayBuilder<FMessageAddress>().Add(RemoteEndpoint->GetAddress()),
+		FTimespan::Zero(), // No Delay
+		FDateTime::MaxValue() // No Expiration
+	);
+
+	// (Re)send any pending reliable messages
+	SendPendingMessages(RemoteEndpoint.ToSharedRef(), ConcertContext.UtcNow);
 }
 
 void FConcertLocalEndpoint::HandleMessage(const FConcertMessageContext& ConcertContext)
diff --git a/Engine/Plugins/Developer/Concert/ConcertMain/Source/ConcertTransport/Public/IConcertMessages.h b/Engine/Plugins/Developer/Concert/ConcertMain/Source/ConcertTransport/Public/IConcertMessages.h
--- a/Engine/Plugins/Developer/Concert/ConcertMain/Source/ConcertTransport/Public/IConcertMessages.h
+++ b/Engine/Plugins/Developer/Concert/ConcertMain/Source/ConcertTransport/Public/IConcertMessages.h
@@ -3,6 +3,7 @@
 #pragma once
 
 #include "ConcertMessageContext.h"
+#include "MessageEndpoint.h"
 
 class IConcertLocalEndpoint;
 
@@ -201,7 +202,7 @@ public:
 	
 	virtual void* ConstructMessage() const override
 	{
-		return new MessageType(MessageTemplate);
+		return FMessageEndpoint::MakeMessage<MessageType>(MessageTemplate);
 	}
 
 	virtual UScriptStruct* GetMessageType() const override
diff --git a/Engine/Plugins/FX/Niagara/Source/Niagara/Private/NiagaraDebuggerClient.cpp b/Engine/Plugins/FX/Niagara/Source/Niagara/Private/NiagaraDebuggerClient.cpp
--- a/Engine/Plugins/FX/Niagara/Source/Niagara/Private/NiagaraDebuggerClient.cpp
+++ b/Engine/Plugins/FX/Niagara/Source/Niagara/Private/NiagaraDebuggerClient.cpp
@@ -284,7 +284,7 @@ bool FNiagaraDebuggerClient::UpdateOutliner(float DeltaSeconds)
 	OutlinerCountdown = 0;
 	if(ensure(Connection.IsValid()))
 	{
-		FNiagaraDebuggerOutlinerUpdate* Message = new FNiagaraDebuggerOutlinerUpdate();
+		FNiagaraDebuggerOutlinerUpdate* Message = FMessageEndpoint::MakeMessage<FNiagaraDebuggerOutlinerUpdate>();
 	
 		//Gather all high level state data to pass to the outliner in the debugger.
 		//TODO: Move out to somewhere neater and add more info.
diff --git a/Engine/Source/Developer/AutomationController/Private/AutomationControllerManger.cpp b/Engine/Source/Developer/AutomationController/Private/AutomationControllerManger.cpp
--- a/Engine/Source/Developer/AutomationController/Private/AutomationControllerManger.cpp
+++ b/Engine/Source/Developer/AutomationController/Private/AutomationControllerManger.cpp
@@ -413,7 +413,7 @@ void FAutomationControllerManager::ProcessComparisonQueue()
 
 			// Send the message back to the automation worker letting it know the results of the comparison test.
 			{
-				FAutomationWorkerImageComparisonResults* Message = new FAutomationWorkerImageComparisonResults(
+				FAutomationWorkerImageComparisonResults* Message = FMessageEndpoint::MakeMessage<FAutomationWorkerImageComparisonResults>(
 					UniqueId,
 					Result.IsNew(),
 					Result.AreSimilar(),
@@ -1294,7 +1294,7 @@ void FAutomationControllerManager::HandleTestDataRequest(const FAutomationWorker
 		}
 	}
 
-	FAutomationWorkerTestDataResponse* ResponseMessage = new FAutomationWorkerTestDataResponse();
+	FAutomationWorkerTestDataResponse* ResponseMessage = FMessageEndpoint::MakeMessage<FAutomationWorkerTestDataResponse>();
 	ResponseMessage->bIsNew = bIsNew;
 	ResponseMessage->JsonData = ResponseJsonData;
 
@@ -1306,7 +1306,7 @@ void FAutomationControllerManager::HandlePerformanceDataRequest(const FAutomatio
 	//TODO Read/Performance data.
 	UE_LOG(LogAutomationController, Log, TEXT("Received PerformanceDataRequest from %s"), *Context->GetSender().ToString());
 
-	FAutomationWorkerPerformanceDataResponse* ResponseMessage = new FAutomationWorkerPerformanceDataResponse();
+	FAutomationWorkerPerformanceDataResponse* ResponseMessage = FMessageEndpoint::MakeMessage<FAutomationWorkerPerformanceDataResponse>();
 	ResponseMessage->bSuccess = true;
 	ResponseMessage->ErrorMessage = TEXT("");
 
diff --git a/Engine/Source/Developer/SlateReflector/Private/WidgetSnapshotService.cpp b/Engine/Source/Developer/SlateReflector/Private/WidgetSnapshotService.cpp
--- a/Engine/Source/Developer/SlateReflector/Private/WidgetSnapshotService.cpp
+++ b/Engine/Source/Developer/SlateReflector/Private/WidgetSnapshotService.cpp
@@ -55,7 +55,7 @@ void FWidgetSnapshotService::HandleWidgetSnapshotRequestMessage(const FWidgetSna
 		FWidgetSnapshotData SnapshotData;
 		SnapshotData.TakeSnapshot(false);
 
-		FWidgetSnapshotResponse* WidgetSnapshotResponse = new FWidgetSnapshotResponse;
+		FWidgetSnapshotResponse* WidgetSnapshotResponse = FMessageEndpoint::MakeMessage<FWidgetSnapshotResponse>();
 		WidgetSnapshotResponse->SnapshotRequestId = Message.SnapshotRequestId;
 		SnapshotData.SaveSnapshotToBuffer(WidgetSnapshotResponse->SnapshotData);
 
diff --git a/Engine/Source/Developer/TargetDeviceServices/Private/Services/TargetDeviceService.cpp b/Engine/Source/Developer/TargetDeviceServices/Private/Services/TargetDeviceService.cpp
--- a/Engine/Source/Developer/TargetDeviceServices/Private/Services/TargetDeviceService.cpp
+++ b/Engine/Source/Developer/TargetDeviceServices/Private/Services/TargetDeviceService.cpp
@@ -413,9 +413,7 @@ void FTargetDeviceService::HandlePingMessage(const FTargetDeviceServicePing& InM
 		const FString& PlatformName = DefaultDevice->GetTargetPlatform().PlatformName();
 		const PlatformInfo::FTargetPlatformInfo* VanillaInfo = DefaultDevice->GetTargetPlatform().GetTargetPlatformInfo().VanillaInfo;
 
-		// message is going to be deleted by FMemory::Free() (see FMessageContext destructor), so allocate it with Malloc
-		void* Memory = FMemory::Malloc(sizeof(FTargetDeviceServicePong), alignof(FTargetDeviceServicePong));
-		FTargetDeviceServicePong* Message = new(Memory) FTargetDeviceServicePong();
+		FTargetDeviceServicePong* Message = FMessageEndpoint::MakeMessage<FTargetDeviceServicePong>();
 
 		Message->Name = DefaultDevice->GetName();
 		Message->Type = TargetDeviceTypes::ToString(DefaultDevice->GetDeviceType());
diff --git a/Engine/Source/Runtime/LiveLinkMessageBusFramework/Private/LiveLinkProvider.cpp b/Engine/Source/Runtime/LiveLinkMessageBusFramework/Private/LiveLinkProvider.cpp
--- a/Engine/Source/Runtime/LiveLinkMessageBusFramework/Private/LiveLinkProvider.cpp
+++ b/Engine/Source/Runtime/LiveLinkMessageBusFramework/Private/LiveLinkProvider.cpp
@@ -156,7 +156,7 @@ private:
 	// Send hierarchy data for named subject
 	void SendSubject(FName SubjectName, const FTrackedSubject& Subject)
 	{
-		FLiveLinkSubjectDataMessage* SubjectData = new FLiveLinkSubjectDataMessage;
+		FLiveLinkSubjectDataMessage* SubjectData = FMessageEndpoint::MakeMessage<FLiveLinkSubjectDataMessage>();
 		SubjectData->RefSkeleton = Subject.RefSkeleton;
 		SubjectData->SubjectName = SubjectName;
 
@@ -174,7 +174,7 @@ private:
 	// Send frame data for named subject
 	void SendSubjectFrame(FName SubjectName, const FTrackedSubject& Subject)
 	{
-		FLiveLinkSubjectFrameMessage* SubjectFrame = new FLiveLinkSubjectFrameMessage;
+		FLiveLinkSubjectFrameMessage* SubjectFrame = FMessageEndpoint::MakeMessage<FLiveLinkSubjectFrameMessage>();
 		SubjectFrame->Transforms = Subject.Transforms;
 		SubjectFrame->SubjectName = SubjectName;
 		SubjectFrame->Curves = Subject.Curves;
@@ -258,8 +258,7 @@ private:
 			MessageAddresses.Add(Address.Address);
 		}
 
-		FLiveLinkClearSubject* ClearSubject = new FLiveLinkClearSubject(SubjectName);
-		MessageEndpoint->Send(ClearSubject, EMessageFlags::Reliable, nullptr, MessageAddresses, FTimespan::Zero(), FDateTime::MaxValue());
+		MessageEndpoint->Send(FMessageEndpoint::MakeMessage<FLiveLinkClearSubject>(SubjectName), EMessageFlags::Reliable, nullptr, MessageAddresses, FTimespan::Zero(), FDateTime::MaxValue());
 	}
 
 public:
diff --git a/Engine/Source/Runtime/MessagingRpc/Private/MessageRpcServer.cpp b/Engine/Source/Runtime/MessagingRpc/Private/MessageRpcServer.cpp
--- a/Engine/Source/Runtime/MessagingRpc/Private/MessageRpcServer.cpp
+++ b/Engine/Source/Runtime/MessagingRpc/Private/MessageRpcServer.cpp
@@ -126,7 +126,7 @@ void FMessageRpcServer::SendProgress(const FGuid& CallId, const FReturnInfo& Ret
 	const TSharedPtr<IAsyncTask>& Task = ReturnInfo.Task;
 
 	MessageEndpoint->Send(
-		new FMessageRpcProgress(
+		FMessageEndpoint::MakeMessage<FMessageRpcProgress>(
 			CallId,
 			Progress.IsValid() ? Progress->GetCompletion().Get(-1.0f) : -1.0f,
 			Progress.IsValid() ? Progress->GetStatusText() : FText::GetEmpty()
diff --git a/Engine/Source/Runtime/MessagingRpc/Public/IMessageRpcClient.h b/Engine/Source/Runtime/MessagingRpc/Public/IMessageRpcClient.h
--- a/Engine/Source/Runtime/MessagingRpc/Public/IMessageRpcClient.h
+++ b/Engine/Source/Runtime/MessagingRpc/Public/IMessageRpcClient.h
@@ -3,6 +3,7 @@
 #pragma once
 
 #include "CoreMinimal.h"
+#include "MessageEndpoint.h"
 #include "Misc/Guid.h"
 #include "IMessageContext.h"
 #include "IMessageRpcCall.h"
@@ -136,7 +137,7 @@ class IMessageRpcClient
 
 		virtual void* ConstructMessage() const override
 		{
-			return new typename RpcType::FRequest(*MessageTemplate);
+			return FMessageEndpoint::MakeMessage<typename RpcType::FRequest>(*MessageTemplate);
 		}
 
 		TFuture<typename RpcType::FResult> GetFuture()
diff --git a/Engine/Source/Runtime/MessagingRpc/Public/IMessageRpcServer.h b/Engine/Source/Runtime/MessagingRpc/Public/IMessageRpcServer.h
--- a/Engine/Source/Runtime/MessagingRpc/Public/IMessageRpcServer.h
+++ b/Engine/Source/Runtime/MessagingRpc/Public/IMessageRpcServer.h
@@ -3,6 +3,7 @@
 
 #include "CoreMinimal.h"
 #include "IMessageContext.h"
+#include "MessageEndpoint.h"
 #include "Async/IAsyncTask.h"
 #include "Async/AsyncResult.h"
 #include "IMessageRpcHandler.h"
@@ -46,7 +47,7 @@ class IMessageRpcServer
 			const TFuture<typename RpcType::FResult>& Future = Result.GetFuture();
 			check(Future.IsReady());
 
-			return new typename RpcType::FResponse(Future.Get());
+			return FMessageEndpoint::MakeMessage<typename RpcType::FResponse>(Future.Get());
 		}
 
 		virtual UScriptStruct* GetResponseTypeInfo() const override
diff --git a/Engine/Source/Runtime/SessionServices/Private/SessionService.cpp b/Engine/Source/Runtime/SessionServices/Private/SessionService.cpp
--- a/Engine/Source/Runtime/SessionServices/Private/SessionService.cpp
+++ b/Engine/Source/Runtime/SessionServices/Private/SessionService.cpp
@@ -83,7 +83,7 @@ void FSessionService::SendLog(const TCHAR* Data, ELogVerbosity::Type Verbosity,
 		if (LogSubscribers.Num() > 0)
 		{
 			MessageEndpoint->Send(
-				new FSessionServiceLog(
+				FMessageEndpoint::MakeMessage<FSessionServiceLog>(
 					Category,
 					Data,
 					FApp::GetInstanceId(),
@@ -105,7 +105,7 @@ void FSessionService::SendNotification(const TCHAR* NotificationText, const FMes
 	}
 
 	MessageEndpoint->Send(
-		new FSessionServiceLog(
+		FMessageEndpoint::MakeMessage<FSessionServiceLog>(
 			FName("RemoteSession"),
 			NotificationText,
 			FApp::GetInstanceId(),
