Concert: Change default validation from Hard to Soft.

This is what our users are normally using, and what we recommend doing anyway.

diff --git a/Engine/Plugins/Developer/Concert/ConcertMain/Source/Concert/Public/ConcertSettings.h b/Engine/Plugins/Developer/Concert/ConcertMain/Source/Concert/Public/ConcertSettings.h
--- a/Engine/Plugins/Developer/Concert/ConcertMain/Source/Concert/Public/ConcertSettings.h
+++ b/Engine/Plugins/Developer/Concert/ConcertMain/Source/Concert/Public/ConcertSettings.h
@@ -280,7 +280,7 @@ struct FConcertSourceControlSettings
 	GENERATED_BODY()
 
 	FConcertSourceControlSettings()
-		: ValidationMode(EConcertSourceValidationMode::Hard)
+		: ValidationMode(EConcertSourceValidationMode::Soft)
 	{}
 
 	UPROPERTY(config, EditAnywhere, Category="Source Control Settings")
