Each archived session should open a separate tab

diff --git a/Engine/Plugins/Developer/Concert/ConcertMain/Source/Concert/Private/ConcertServer.cpp b/Engine/Plugins/Developer/Concert/ConcertMain/Source/Concert/Private/ConcertServer.cpp
--- a/Engine/Plugins/Developer/Concert/ConcertMain/Source/Concert/Private/ConcertServer.cpp
+++ b/Engine/Plugins/Developer/Concert/ConcertMain/Source/Concert/Private/ConcertServer.cpp
@@ -26,13 +26,13 @@
 namespace ConcertServerUtil
 {
 
-const TCHAR* GetServerSystemMutexName()
+static const TCHAR* GetServerSystemMutexName()
 {
 	// A system wide mutex name used by this application instances that will unlikely be found in other applications.
 	return TEXT("Unreal_ConcertServer_67822dAB");
 }
 
-FString GetArchiveName(const FString& SessionName, const FConcertSessionSettings& Settings)
+static FString GetArchiveName(const FString& SessionName, const FConcertSessionSettings& Settings)
 {
 	if (Settings.ArchiveNameOverride.IsEmpty())
 	{
@@ -44,12 +44,12 @@ FString GetArchiveName(const FString& SessionName, const FConcertSessionSettings
 	}
 }
 
-FString GetSessionRepositoryDatabasePathname(const FString& Role)
+static FString GetSessionRepositoryDatabasePathname(const FString& Role)
 {
 	return FPaths::ProjectSavedDir() / Role / TEXT("Repositories.json");
 }
 
-bool SaveSessionRepositoryDatabase(const FString& Role, const FConcertServerSessionRepositoryDatabase& RepositoryDb)
+static bool SaveSessionRepositoryDatabase(const FString& Role, const FConcertServerSessionRepositoryDatabase& RepositoryDb)
 {
 	if (TUniquePtr<FArchive> FileWriter = TUniquePtr<FArchive>(IFileManager::Get().CreateFileWriter(*GetSessionRepositoryDatabasePathname(Role))))
 	{
@@ -63,7 +63,7 @@ bool SaveSessionRepositoryDatabase(const FString& Role, const FConcertServerSess
 	return false;
 }
 
-bool LoadSessionRepositoryDatabase(const FString& Role, FConcertServerSessionRepositoryDatabase& RepositoryDb)
+static bool LoadSessionRepositoryDatabase(const FString& Role, FConcertServerSessionRepositoryDatabase& RepositoryDb)
 {
 	if (TUniquePtr<FArchive> FileReader = TUniquePtr<FArchive>(IFileManager::Get().CreateFileReader(*GetSessionRepositoryDatabasePathname(Role))))
 	{
@@ -76,7 +76,6 @@ bool LoadSessionRepositoryDatabase(const FString& Role, FConcertServerSessionRep
 
 	return false;
 }
-
 }
 
 FConcertServer::FConcertServer(const FString& InRole, const FConcertSessionFilter& InAutoArchiveSessionFilter, IConcertServerEventSink* InEventSink, const TSharedPtr<IConcertEndpointProvider>& InEndpointProvider)
@@ -1521,6 +1520,8 @@ TSharedPtr<IConcertServerSession> FConcertServer::RestoreArchivedSession(const F
 		}
 
 		FConcertSessionInfo LiveSessionInfo = NewSessionInfo;
+		// Detect restoring the same archived session twice by hashing archived ID
+		LiveSessionInfo.SessionId = FGuid::NewGuid();
 		LiveSessionInfo.SessionName = MoveTemp(LiveSessionName);
 		LiveSessionInfo.VersionInfos = ArchivedSessionInfo->VersionInfos;
 
