Introduce transient state on the server state.

diff --git a/Engine/Plugins/Developer/Concert/ConcertMain/Source/Concert/Private/ConcertServerSession.cpp b/Engine/Plugins/Developer/Concert/ConcertMain/Source/Concert/Private/ConcertServerSession.cpp
--- a/Engine/Plugins/Developer/Concert/ConcertMain/Source/Concert/Private/ConcertServerSession.cpp
+++ b/Engine/Plugins/Developer/Concert/ConcertMain/Source/Concert/Private/ConcertServerSession.cpp
@@ -2,6 +2,7 @@
 
 #include "ConcertServerSession.h"
 #include "ConcertLogGlobal.h"
+#include "ConcertMessages.h"
 #include "IConcertEndpoint.h"
 #include "Scratchpad/ConcertScratchpad.h"
 
@@ -177,12 +178,17 @@ void FConcertServerSession::HandleDiscoverAndJoinSessionEvent(const FConcertMess
 	JoinReply.ConcertProtocolVersion = EConcertMessageVersion::LatestVersion;
 	JoinReply.SessionServerEndpointId = SessionInfo.ServerEndpointId;
 
+	const FConcertSessionInfo& InSessionInfo = GetSessionInfo();
 	if (SessionClients.Contains(Context.SenderConcertEndpointId))
 	{
 		JoinReply.ConnectionResult = EConcertConnectionResult::AlreadyConnected;
 		JoinReply.SessionClients = GetSessionClients();
 	}
-	// TODO: check connection requirement
+	else if (InSessionInfo.State == EConcertSessionState::Transient)
+	{
+		JoinReply.ConnectionResult = EConcertConnectionResult::ConnectionRefused;
+	}
+		// TODO: check connection requirement
 	else // if (CheckConnectionRequirement(Message->ClientInfo))
 	{
 		// Accept the connection
@@ -319,4 +325,4 @@ void FConcertServerSession::TickConnections(float DeltaSeconds)
 {
 	// External callback
 	OnTickDelegate.Broadcast(*this, DeltaSeconds);
-}
\ No newline at end of file
+}
diff --git a/Engine/Plugins/Developer/Concert/ConcertMain/Source/Concert/Public/ConcertMessageData.h b/Engine/Plugins/Developer/Concert/ConcertMain/Source/Concert/Public/ConcertMessageData.h
--- a/Engine/Plugins/Developer/Concert/ConcertMain/Source/Concert/Public/ConcertMessageData.h
+++ b/Engine/Plugins/Developer/Concert/ConcertMain/Source/Concert/Public/ConcertMessageData.h
@@ -22,6 +22,17 @@ enum class EConcertServerFlags : uint8
 };
 ENUM_CLASS_FLAGS(EConcertServerFlags)
 
+
+UENUM()
+enum class EConcertSessionState : uint8
+{
+	/** Session is a normal state and can be joined. */
+	Normal = 0,
+
+	/** Session is in a transient state and cannot be joined. */
+	Transient
+};
+
 UENUM()
 enum class EConcertPayloadCompressionType : uint8
 {
@@ -203,13 +214,17 @@ struct FConcertSessionInfo
 	UPROPERTY(VisibleAnywhere, Category = "Session Info")
 	FString OwnerDeviceName;
 
-	/** Settings pertaining to project, change list number etc */ 
+	/** Settings pertaining to project, change list number etc */
 	UPROPERTY(VisibleAnywhere, Category="Session Info")
 	FConcertSessionSettings Settings;
 
 	/** Version information for this session. This is set during creation, and updated each time the session is restored */
 	UPROPERTY()
 	TArray<FConcertSessionVersionInfo> VersionInfos;
+
+	/** Current state of the session used to determine joinability by clients. */
+	UPROPERTY()
+	EConcertSessionState State = EConcertSessionState::Normal;
 };
 
 /** Holds filter rules used when migrating session data */
