Add LastModified property to FConcertSessionInfo

diff --git a/Engine/Plugins/Developer/Concert/ConcertMain/Source/Concert/Private/ConcertServer.cpp b/Engine/Plugins/Developer/Concert/ConcertMain/Source/Concert/Private/ConcertServer.cpp
--- a/Engine/Plugins/Developer/Concert/ConcertMain/Source/Concert/Private/ConcertServer.cpp
+++ b/Engine/Plugins/Developer/Concert/ConcertMain/Source/Concert/Private/ConcertServer.cpp
@@ -463,8 +463,10 @@ void FConcertServer::RecoverSessions(const FConcertServerSessionRepository& InRe
 {
 	// Find any existing live sessions to automatically restore when recovering from an improper server shutdown
 	TArray<FConcertSessionInfo> LiveSessionInfos;
-	EventSink->GetSessionsFromPath(*this, InRepository.WorkingDir, LiveSessionInfos);
-
+	TArray<FDateTime> LiveSessionCreationTimes;
+	EventSink->GetSessionsFromPath(*this, InRepository.WorkingDir, LiveSessionInfos, &LiveSessionCreationTimes);
+	UpdateLastModified(LiveSessionInfos, LiveSessionCreationTimes);
+	
 	// Restore any existing live sessions
 	for (FConcertSessionInfo& LiveSessionInfo : LiveSessionInfos)
 	{
@@ -489,7 +491,8 @@ void FConcertServer::RecoverSessions(const FConcertServerSessionRepository& InRe
 		// In theory, archives are immutable, but the server will end up touching the files and change the 'modification time'. Ensure to look at 'creation time'.
 		EventSink->GetSessionsFromPath(*this, InRepository.SavedDir, ArchivedSessionInfos, &ArchivedSessionCreationTimes);
 		check(ArchivedSessionInfos.Num() == ArchivedSessionCreationTimes.Num());
-
+		UpdateLastModified(ArchivedSessionInfos, ArchivedSessionCreationTimes);
+		
 		// Trim the oldest archived sessions.
 		if (bCleanupExpiredSessions && Settings->NumSessionsToKeep > 0 && ArchivedSessionInfos.Num() > Settings->NumSessionsToKeep)
 		{
@@ -541,6 +544,14 @@ void FConcertServer::RecoverSessions(const FConcertServerSessionRepository& InRe
 	}
 }
 
+void FConcertServer::UpdateLastModified(TArray<FConcertSessionInfo>& SessionInfos, const TArray<FDateTime>& SessionCreationTimes)
+{
+	for (int32 i = 0; i < SessionInfos.Num(); ++i)
+	{
+		SessionInfos[i].LastModified = SessionCreationTimes[i];
+	}
+}
+
 void FConcertServer::ArchiveOfflineSessions(const FConcertServerSessionRepository& InRepository)
 {
 	// Find existing live session files to automatically archive them when recovering from an improper server shutdown.
@@ -554,6 +565,7 @@ void FConcertServer::ArchiveOfflineSessions(const FConcertServerSessionRepositor
 		FConcertSessionInfo ArchivedSessionInfo = LiveSessionInfo;
 		ArchivedSessionInfo.SessionId = FGuid::NewGuid();
 		ArchivedSessionInfo.SessionName = ConcertServerUtil::GetArchiveName(LiveSessionInfo.SessionName, LiveSessionInfo.Settings);
+		ArchivedSessionInfo.SetLastModifiedToNow();
 
 		if (EventSink->ArchiveSession(*this, InRepository.GetSessionWorkingDir(LiveSessionInfo.SessionId), InRepository.GetSessionSavedDir(ArchivedSessionInfo.SessionId), ArchivedSessionInfo, AutoArchiveSessionFilter))
 		{
@@ -1409,7 +1421,10 @@ TSharedPtr<IConcertServerSession> FConcertServer::CreateLiveSession(const FConce
 		InRepository.GetSessionWorkingDir(LiveSessionInfo.SessionId)
 		);
 
-	if (EventSink->OnLiveSessionCreated(*this, LiveSession.ToSharedRef())) // EventSync could complete the session initialization?
+	FInternalLiveSessionCreationParams CreationParams;
+	CreationParams.OnModifiedCallback.BindSP(LiveSession.Get(), &FConcertServerSession::SetLastModifiedToNow);
+	
+	if (EventSink->OnLiveSessionCreated(*this, LiveSession.ToSharedRef(), CreationParams)) // EventSync could complete the session initialization?
 	{
 		LiveSessions.Add(LiveSessionInfo.SessionId, LiveSession);
 		LiveSession->Startup();
@@ -1524,6 +1539,7 @@ TSharedPtr<IConcertServerSession> FConcertServer::RestoreArchivedSession(const F
 		LiveSessionInfo.SessionId = FGuid::NewGuid();
 		LiveSessionInfo.SessionName = MoveTemp(LiveSessionName);
 		LiveSessionInfo.VersionInfos = ArchivedSessionInfo->VersionInfos;
+		LiveSessionInfo.SetLastModifiedToNow();
 
 		// Ensure the new version is compatible with the old version, and append this new version if it is different to the last used version
 		// Note: Older archived sessions didn't used to have any version info stored for them, and the version info may be missing completely when using -CONCERTIGNORE
diff --git a/Engine/Plugins/Developer/Concert/ConcertMain/Source/Concert/Private/ConcertServer.h b/Engine/Plugins/Developer/Concert/ConcertMain/Source/Concert/Private/ConcertServer.h
--- a/Engine/Plugins/Developer/Concert/ConcertMain/Source/Concert/Private/ConcertServer.h
+++ b/Engine/Plugins/Developer/Concert/ConcertMain/Source/Concert/Private/ConcertServer.h
@@ -113,6 +113,9 @@ private:
 
 	/** Recover the sessions found in the working directory into live session, build the list of archived sessions and rotate them, keeping only the N most recent. */
 	void RecoverSessions(const FConcertServerSessionRepository& InRepository, bool bCleanupExpiredSessions);
+	
+	/** Sets each FConcertSessionInfo::LastModified time to the corresponding index in SessionCreationTimes. */
+	void UpdateLastModified(TArray<FConcertSessionInfo>& SessionInfos, const TArray<FDateTime>& SessionCreationTimes);
 
 	/**
 	 * Migrate the live sessions from the working directory (before sessions being recovered into live one) to the archived directory.
diff --git a/Engine/Plugins/Developer/Concert/ConcertMain/Source/Concert/Private/ConcertServerSession.h b/Engine/Plugins/Developer/Concert/ConcertMain/Source/Concert/Private/ConcertServerSession.h
--- a/Engine/Plugins/Developer/Concert/ConcertMain/Source/Concert/Private/ConcertServerSession.h
+++ b/Engine/Plugins/Developer/Concert/ConcertMain/Source/Concert/Private/ConcertServerSession.h
@@ -11,7 +11,10 @@ class IConcertLocalEndpoint;
 struct FConcertServerSettings;
 
 /** Implementation of a Concert Server session */
-class FConcertServerSession : public IConcertServerSession, private FConcertSessionCommonImpl
+class FConcertServerSession
+	: public IConcertServerSession
+	, public TSharedFromThis<FConcertServerSession>
+	, private FConcertSessionCommonImpl
 {
 public:
 	FConcertServerSession(const FConcertSessionInfo& InSessionInfo, const FConcertServerSettings& InSettings, TSharedPtr<IConcertLocalEndpoint> InServerSessionEndpoint, const FString& InSessionDirectory);
@@ -36,6 +39,11 @@ public:
 		SendSessionNameChanged();
 	}
 
+	void SetLastModifiedToNow()
+	{
+		SessionInfo.SetLastModifiedToNow();
+	}
+
 	virtual const FConcertSessionInfo& GetSessionInfo() const override
 	{
 		return CommonGetSessionInfo();
diff --git a/Engine/Plugins/Developer/Concert/ConcertMain/Source/Concert/Public/ConcertMessageData.h b/Engine/Plugins/Developer/Concert/ConcertMain/Source/Concert/Public/ConcertMessageData.h
--- a/Engine/Plugins/Developer/Concert/ConcertMain/Source/Concert/Public/ConcertMessageData.h
+++ b/Engine/Plugins/Developer/Concert/ConcertMain/Source/Concert/Public/ConcertMessageData.h
@@ -218,6 +218,10 @@ struct FConcertSessionInfo
 	UPROPERTY(VisibleAnywhere, Category="Session Info")
 	FConcertSessionSettings Settings;
 
+	/** The last time the session directory was modified in local time. */
+	UPROPERTY(VisibleAnywhere, Category="Session Info")
+	FDateTime LastModified;
+
 	/** Version information for this session. This is set during creation, and updated each time the session is restored */
 	UPROPERTY()
 	TArray<FConcertSessionVersionInfo> VersionInfos;
@@ -225,6 +229,11 @@ struct FConcertSessionInfo
 	/** Current state of the session used to determine joinability by clients. */
 	UPROPERTY()
 	EConcertSessionState State = EConcertSessionState::Normal;
+
+	void SetLastModifiedToNow()
+	{
+		LastModified = FDateTime::Now();
+	}
 };
 
 /** Holds filter rules used when migrating session data */
diff --git a/Engine/Plugins/Developer/Concert/ConcertMain/Source/Concert/Public/IConcertServerEventSink.h b/Engine/Plugins/Developer/Concert/ConcertMain/Source/Concert/Public/IConcertServerEventSink.h
--- a/Engine/Plugins/Developer/Concert/ConcertMain/Source/Concert/Public/IConcertServerEventSink.h
+++ b/Engine/Plugins/Developer/Concert/ConcertMain/Source/Concert/Public/IConcertServerEventSink.h
@@ -11,6 +11,14 @@ class IConcertServerSession;
 struct FConcertSessionInfo;
 struct FConcertSessionFilter;
 
+struct FInternalLiveSessionCreationParams
+{
+	DECLARE_DELEGATE(FOnSessionModified)
+
+	/** Called when the session was modified */
+	FOnSessionModified OnModifiedCallback;
+};
+
 /** Interface for events that Concert server can emit */
 class IConcertServerEventSink
 {
@@ -25,7 +33,7 @@ public:
 	 * @note This function is called for both newly created sessions and after recovering a live session during server start-up.
 	 * @return true if the session creation could be completed without error, false otherwise (ex if the database fails to open).
 	 */
-	virtual bool OnLiveSessionCreated(const IConcertServer& InServer, TSharedRef<IConcertServerSession> InLiveSession) = 0;
+	virtual bool OnLiveSessionCreated(const IConcertServer& InServer, TSharedRef<IConcertServerSession> InLiveSession, const FInternalLiveSessionCreationParams& AdditionalParams) = 0;
 	
 	/**
 	 * Called before the session is destroyed (and before Shutdown is called on it).
diff --git a/Engine/Plugins/Developer/Concert/ConcertSync/ConcertSyncServer/Source/ConcertSyncServer/Private/ConcertServerEventForwardingSink.h b/Engine/Plugins/Developer/Concert/ConcertSync/ConcertSyncServer/Source/ConcertSyncServer/Private/ConcertServerEventForwardingSink.h
--- a/Engine/Plugins/Developer/Concert/ConcertSync/ConcertSyncServer/Source/ConcertSyncServer/Private/ConcertServerEventForwardingSink.h
+++ b/Engine/Plugins/Developer/Concert/ConcertSync/ConcertSyncServer/Source/ConcertSyncServer/Private/ConcertServerEventForwardingSink.h
@@ -24,9 +24,10 @@ public:
 		This()->GetSessionsFromPathImpl(InServer, InPath, OutSessionInfos, OutSessionCreationTimes);
 	}
 
-	virtual bool OnLiveSessionCreated(const IConcertServer& InServer, TSharedRef<IConcertServerSession> InLiveSession) final
+	virtual bool OnLiveSessionCreated(const IConcertServer& InServer, TSharedRef<IConcertServerSession> InLiveSession, const FInternalLiveSessionCreationParams& AdditionalParams) final
 	{
-		const bool bSuccess = This()->OnLiveSessionCreatedImpl(InServer, InLiveSession);
+		const bool bSuccess = This()->OnLiveSessionCreatedImpl(InServer, InLiveSession, AdditionalParams);
+		// AdditionalParams are internal use only - they should never be exposed publicly!
 		ConcertServerEvents::OnLiveSessionCreated().Broadcast(bSuccess, InServer, InLiveSession);
 		return bSuccess;
 	}
diff --git a/Engine/Plugins/Developer/Concert/ConcertSync/ConcertSyncServer/Source/ConcertSyncServer/Private/ConcertSyncServer.cpp b/Engine/Plugins/Developer/Concert/ConcertSync/ConcertSyncServer/Source/ConcertSyncServer/Private/ConcertSyncServer.cpp
--- a/Engine/Plugins/Developer/Concert/ConcertSync/ConcertSyncServer/Source/ConcertSyncServer/Private/ConcertSyncServer.cpp
+++ b/Engine/Plugins/Developer/Concert/ConcertSync/ConcertSyncServer/Source/ConcertSyncServer/Private/ConcertSyncServer.cpp
@@ -308,10 +308,10 @@ void FConcertSyncServer::GetSessionsFromPathImpl(const IConcertServer& InServer,
 	});
 }
 
-bool FConcertSyncServer::OnLiveSessionCreatedImpl(const IConcertServer& InServer, TSharedRef<IConcertServerSession> InSession)
+bool FConcertSyncServer::OnLiveSessionCreatedImpl(const IConcertServer& InServer, TSharedRef<IConcertServerSession> InSession, const FInternalLiveSessionCreationParams& AdditionalParams)
 {
 	ConcertSyncServerUtils::WriteSessionInfoToDirectory(InSession->GetSessionWorkingDirectory(), InSession->GetSessionInfo());
-	return CreateLiveSession(InSession);
+	return CreateLiveSession(InSession, AdditionalParams);
 }
 
 void FConcertSyncServer::OnLiveSessionDestroyedImpl(const IConcertServer& InServer, TSharedRef<IConcertServerSession> InSession)
@@ -515,7 +515,7 @@ void FConcertSyncServer::DestroySequencerManager(const TSharedRef<FConcertSyncSe
 	LiveSessionSequencerManagers.Remove(InLiveSession->GetSession().GetId());
 }
 
-bool FConcertSyncServer::CreateLiveSession(const TSharedRef<IConcertServerSession>& InSession)
+bool FConcertSyncServer::CreateLiveSession(const TSharedRef<IConcertServerSession>& InSession, const FInternalLiveSessionCreationParams& AdditionalParams)
 {
 	DestroyLiveSession(InSession);
 
@@ -529,6 +529,12 @@ bool FConcertSyncServer::CreateLiveSession(const TSharedRef<IConcertServerSessio
 			CreateSequencerManager(LiveSession.ToSharedRef());
 		}
 
+		// We needn't call OnActivityProduced().Remove(...) because the subscription needs to stay for the lifetime of FConcertSyncServerLiveSession::SessionDatabase
+		check(AdditionalParams.OnModifiedCallback.IsBound());
+		LiveSession->GetSessionDatabase().OnActivityProduced().AddLambda([OnSessionModifiedCallback = AdditionalParams.OnModifiedCallback](const FConcertSyncActivity&)
+		{
+			OnSessionModifiedCallback.Execute();
+		});
 		return true;
 	}
 
diff --git a/Engine/Plugins/Developer/Concert/ConcertSync/ConcertSyncServer/Source/ConcertSyncServer/Private/ConcertSyncServer.h b/Engine/Plugins/Developer/Concert/ConcertSync/ConcertSyncServer/Source/ConcertSyncServer/Private/ConcertSyncServer.h
--- a/Engine/Plugins/Developer/Concert/ConcertSync/ConcertSyncServer/Source/ConcertSyncServer/Private/ConcertSyncServer.h
+++ b/Engine/Plugins/Developer/Concert/ConcertSync/ConcertSyncServer/Source/ConcertSyncServer/Private/ConcertSyncServer.h
@@ -34,7 +34,7 @@ public:
 
 	//~ IConcertServerEventSink interface
 	void GetSessionsFromPathImpl(const IConcertServer& InServer, const FString& InPath, TArray<FConcertSessionInfo>& OutSessionInfos, TArray<FDateTime>* OutSessionCreationTimes = nullptr);
-	bool OnLiveSessionCreatedImpl(const IConcertServer& InServer, TSharedRef<IConcertServerSession> InLiveSession);
+	bool OnLiveSessionCreatedImpl(const IConcertServer& InServer, TSharedRef<IConcertServerSession> InLiveSession, const FInternalLiveSessionCreationParams& AdditionalParams);
 	void OnLiveSessionDestroyedImpl(const IConcertServer& InServer, TSharedRef<IConcertServerSession> InLiveSession);
 	bool OnArchivedSessionCreatedImpl(const IConcertServer& InServer, const FString& InArchivedSessionRoot, const FConcertSessionInfo& InArchivedSessionInfo);
 	void OnArchivedSessionDestroyedImpl(const IConcertServer& InServer, const FGuid& InArchivedSessionId);
@@ -54,7 +54,7 @@ private:
 	void CreateSequencerManager(const TSharedRef<FConcertSyncServerLiveSession>& InLiveSession);
 	void DestroySequencerManager(const TSharedRef<FConcertSyncServerLiveSession>& InLiveSession);
 
-	bool CreateLiveSession(const TSharedRef<IConcertServerSession>& InSession);
+	bool CreateLiveSession(const TSharedRef<IConcertServerSession>& InSession, const FInternalLiveSessionCreationParams& AdditionalParams);
 	void DestroyLiveSession(const TSharedRef<IConcertServerSession>& InSession);
 
 	bool CreateArchivedSession(const FString& InArchivedSessionRoot, const FConcertSessionInfo& InArchivedSessionInfo);
