Fix cook error with uninitialized value.

diff --git a/Engine/Plugins/Developer/Concert/ConcertMain/Source/Concert/Public/ConcertMessages.h b/Engine/Plugins/Developer/Concert/ConcertMain/Source/Concert/Public/ConcertMessages.h
--- a/Engine/Plugins/Developer/Concert/ConcertMain/Source/Concert/Public/ConcertMessages.h
+++ b/Engine/Plugins/Developer/Concert/ConcertMain/Source/Concert/Public/ConcertMessages.h
@@ -464,7 +464,7 @@ struct FConcertAdmin_BatchDeleteSessionRequest : public FConcertRequestData
 
 	/** How the operation input is supposed to be interpreted */
 	UPROPERTY(VisibleAnywhere, Category = "Concert Message")
-	EBatchSessionDeletionFlags Flags;
+	EBatchSessionDeletionFlags Flags = EBatchSessionDeletionFlags::Strict;
 
 	//For now only the user name and device name of the client is used to id him as the owner of a session
 	UPROPERTY(VisibleAnywhere, Category = "Concert Message")
