Correct inaccurate documentation

diff --git a/Engine/Plugins/Developer/Concert/ConcertMain/Source/Concert/Public/IConcertClient.h b/Engine/Plugins/Developer/Concert/ConcertMain/Source/Concert/Public/IConcertClient.h
--- a/Engine/Plugins/Developer/Concert/ConcertMain/Source/Concert/Public/IConcertClient.h
+++ b/Engine/Plugins/Developer/Concert/ConcertMain/Source/Concert/Public/IConcertClient.h
@@ -344,11 +344,11 @@ public:
 	virtual TFuture<EConcertResponseCode> DeleteSession(const FGuid& ServerAdminEndpointId, const FGuid& SessionId) = 0;
 
 	/**
-	 * Deletes several live or archives sessions from the server. The client will locally filter SessionIds and request deletion of only those sessions which the client is an owner of.
-	 * The server will still internally validate the client has permission and possibly refuse if there's a mismatch.
+	 * Deletes several live or archives sessions from the server. If the client is not sure whether the client has permission for it,
+	 * it can set the SkipForbiddenSessions - the server will skip disallowed sessions instead of rejecting the entire operation.
 	 * @param ServerAdminEndpointId	The Id of the Concert Server query endpoint
 	 * @param BatchDeletionArgs		The arguments that will be used for batch deleting the session(s)
-	 * @return A future that will contain the server's response AND which sessions were filtered out by the client before sending.
+	 * @return A future that will contain the server's response (and optionally which sessions were skipped due to permissions)
 	 */
 	virtual TFuture<FConcertAdmin_BatchDeleteSessionResponse> BatchDeleteSessions(const FGuid& ServerAdminEndpointId, const FConcertBatchDeleteSessionsArgs& BatchDeletionArgs) = 0;
 	
