Expose concert logger publicly

diff --git a/Engine/Plugins/Developer/Concert/ConcertMain/Source/Concert/Private/ConcertClient.cpp b/Engine/Plugins/Developer/Concert/ConcertMain/Source/Concert/Private/ConcertClient.cpp
--- a/Engine/Plugins/Developer/Concert/ConcertMain/Source/Concert/Private/ConcertClient.cpp
+++ b/Engine/Plugins/Developer/Concert/ConcertMain/Source/Concert/Private/ConcertClient.cpp
@@ -6,6 +6,7 @@
 #include "ConcertUtil.h"
 #include "ConcertLogger.h"
 #include "ConcertLogGlobal.h"
+#include "ConcertTransportEvents.h"
 
 #include "Algo/Transform.h"
 
@@ -712,7 +713,13 @@ void FConcertClient::Startup()
 	if (!ClientAdminEndpoint.IsValid() && EndpointProvider.IsValid())
 	{
 		// Create the client administration endpoint
-		ClientAdminEndpoint = EndpointProvider->CreateLocalEndpoint(TEXT("Admin"), Settings->EndpointSettings, &FConcertLogger::CreateLogger);
+		ClientAdminEndpoint = EndpointProvider->CreateLocalEndpoint(TEXT("Admin"), Settings->EndpointSettings, [this](const FConcertEndpointContext& Context)
+		{
+			return FConcertLogger::CreateLogger(Context, [this](const FConcertLog& Log)
+			{
+				ConcertTransportEvents::OnConcertClientLogEvent().Broadcast(*this, Log);
+			});
+		});
 	}
 
 	FCoreDelegates::OnEndFrame.AddRaw(this, &FConcertClient::OnEndFrame);
@@ -1375,7 +1382,13 @@ TFuture<EConcertResponseCode> FConcertClient::CreateClientSession(const FConcert
 		SessionInfo, 
 		ClientInfo, 
 		Settings->ClientSettings, 
-		EndpointProvider->CreateLocalEndpoint(SessionInfo.SessionName, Settings->EndpointSettings, &FConcertLogger::CreateLogger),
+		EndpointProvider->CreateLocalEndpoint(SessionInfo.SessionName, Settings->EndpointSettings, [this](const FConcertEndpointContext& Context)
+		{
+			return FConcertLogger::CreateLogger(Context, [this](const FConcertLog& Log)
+			{
+				ConcertTransportEvents::OnConcertClientLogEvent().Broadcast(*this, Log);
+			});
+		}),
 		Paths.GetSessionWorkingDir(SessionInfo.SessionId)
 		);
 	OnSessionStartupDelegate.Broadcast(ClientSession.ToSharedRef());
diff --git a/Engine/Plugins/Developer/Concert/ConcertMain/Source/Concert/Private/ConcertLogger.cpp b/Engine/Plugins/Developer/Concert/ConcertMain/Source/Concert/Private/ConcertLogger.cpp
--- a/Engine/Plugins/Developer/Concert/ConcertMain/Source/Concert/Private/ConcertLogger.cpp
+++ b/Engine/Plugins/Developer/Concert/ConcertMain/Source/Concert/Private/ConcertLogger.cpp
@@ -5,7 +5,9 @@
 #include "HAL/IConsoleManager.h"
 #include "IConcertMessages.h"
 #include "ConcertLogGlobal.h"
+#include "ConcertTransportEvents.h"
 #include "ConcertUtil.h"
+#include "Algo/AllOf.h"
 
 #include "Logging/LogVerbosity.h"
 #include "Misc/App.h"
@@ -188,6 +190,8 @@ struct FActiveLoggers
 				}
 			}
 		}
+
+		ConcertTransportEvents::OnConcertTransportLoggingEnabledChangedEvent().Broadcast(true);
 	}
 
 	/** Disable additional concert logging. */
@@ -203,6 +207,17 @@ struct FActiveLoggers
 				}
 			}
 		}
+		
+		ConcertTransportEvents::OnConcertTransportLoggingEnabledChangedEvent().Broadcast(false);
+	}
+
+	bool AreAllLoggersEnabled() const
+	{
+		return Algo::AllOf(Loggers, [](const TWeakPtr<FConcertLogger, ESPMode::ThreadSafe>& Logger)
+		{
+			TSharedPtr<FConcertLogger, ESPMode::ThreadSafe> LoggerPtr = Logger.Pin();
+			return !LoggerPtr.IsValid() && LoggerPtr->IsLogging();
+		});
 	}
 
 	void ChangeVerbosity(const TArray<FString>& Args)
@@ -248,10 +263,23 @@ FActiveLoggers& GetActiveLoggers()
 	return ActiveLoggers;
 }
 
-IConcertTransportLoggerRef FConcertLogger::CreateLogger(const FConcertEndpointContext& InOwnerContext)
+namespace ConcertTransportEvents
+{
+	bool IsLoggingEnabled()
+	{
+		return GetActiveLoggers().AreAllLoggersEnabled();
+	}
+	
+	void SetLoggingEnabled(bool bEnabled)
+	{
+		return GetActiveLoggers().EnableLogging();
+	}
+}
+
+IConcertTransportLoggerRef FConcertLogger::CreateLogger(const FConcertEndpointContext& InOwnerContext, FLogListener LogListenerFunc)
 {
 	FActiveLoggers& ActiveLoggers = GetActiveLoggers();
-	TSharedPtr<FConcertLogger, ESPMode::ThreadSafe> SharedLogger = MakeShared<FConcertLogger, ESPMode::ThreadSafe>(InOwnerContext);
+	TSharedPtr<FConcertLogger, ESPMode::ThreadSafe> SharedLogger = MakeShared<FConcertLogger, ESPMode::ThreadSafe>(InOwnerContext, MoveTemp(LogListenerFunc));
 	ActiveLoggers.Loggers.Add(SharedLogger);
 	return SharedLogger.ToSharedRef();
 }
@@ -272,11 +300,11 @@ void FConcertLogger::SetVerboseLogging(bool bInState)
 	}
 }
 
-FConcertLogger::FConcertLogger(const FConcertEndpointContext& InOwnerContext)
+FConcertLogger::FConcertLogger(const FConcertEndpointContext& InOwnerContext, FLogListener LogListenerFunc)
 	: bIsLogging(false)
 	, OwnerContext(InOwnerContext)
-{
-}
+	, LogListenerFunc(MoveTemp(LogListenerFunc))
+{}
 
 FConcertLogger::~FConcertLogger()
 {
@@ -714,6 +742,9 @@ void FConcertLogger::LogEntry(FConcertLog& Log)
 		ConcertLoggerUtil::PopulateLogMessagePayload(Log.SerializedPayload, Log);
 	}
 
+	// Inform external systems about the log
+	LogListenerFunc(Log);
+
 	FString CSVRow;
 	for (TFieldIterator<const FProperty> PropertyIt(FConcertLog::StaticStruct(), EFieldIteratorFlags::IncludeSuper, EFieldIteratorFlags::ExcludeDeprecated, EFieldIteratorFlags::IncludeInterfaces); PropertyIt; ++PropertyIt)
 	{
diff --git a/Engine/Plugins/Developer/Concert/ConcertMain/Source/Concert/Private/ConcertLogger.h b/Engine/Plugins/Developer/Concert/ConcertMain/Source/Concert/Private/ConcertLogger.h
--- a/Engine/Plugins/Developer/Concert/ConcertMain/Source/Concert/Private/ConcertLogger.h
+++ b/Engine/Plugins/Developer/Concert/ConcertMain/Source/Concert/Private/ConcertLogger.h
@@ -7,83 +7,24 @@
 #include "IConcertEndpoint.h"
 #include "ConcertMessageData.h"
 #include "IConcertTransportLogger.h"
-#include "ConcertLogger.generated.h"
 
+struct FConcertLog;
 struct FConcertMessageContext;
 
-UENUM()
-enum class EConcertLogMessageAction : uint8
-{
-	None,
-	Send,
-	Publish,
-	Receive,
-	Queue,
-	Discard,
-	Duplicate,
-	TimeOut,
-	Process,
-	EndpointDiscovery,
-	EndpointTimeOut,
-	EndpointClosure,
-};
-
-USTRUCT()
-struct FConcertLog
-{
-	GENERATED_BODY()
-
-	UPROPERTY()
-	uint64 Frame = 0;
-
-	UPROPERTY()
-	FGuid MessageId;
-
-	UPROPERTY()
-	uint16 MessageOrderIndex  = 0;
-
-	UPROPERTY()
-	uint16 ChannelId = 0;
-
-	UPROPERTY()
-	FDateTime Timestamp = {0};
-
-	UPROPERTY()
-	EConcertLogMessageAction MessageAction = EConcertLogMessageAction::None;
-
-	UPROPERTY()
-	FName MessageTypeName;
-
-	UPROPERTY()
-	FGuid OriginEndpointId;
-
-	UPROPERTY()
-	FGuid DestinationEndpointId;
-
-	UPROPERTY()
-	FName CustomPayloadTypename;
-
-	UPROPERTY()
-	int32 CustomPayloadUncompressedByteSize = 0;
-
-	UPROPERTY()
-	FString StringPayload;
-
-	UPROPERTY(Transient)
-	FConcertSessionSerializedPayload SerializedPayload;
-};
-
 class FConcertLogger : public IConcertTransportLogger
 {
 public:
+
+	using FLogListener = TFunction<void(const FConcertLog&)>;
+	
 	/** Factory function for use with FConcertTransportLoggerFactory */
-	static IConcertTransportLoggerRef CreateLogger(const FConcertEndpointContext& InOwnerContext);
+	static IConcertTransportLoggerRef CreateLogger(const FConcertEndpointContext& InOwnerContext, FLogListener LogListenerFunc);
 
 	/** Static function to enable / disable verbose logging. */
 	static void SetVerboseLogging(bool bInState);
 
-	explicit FConcertLogger(const FConcertEndpointContext& InOwnerContext);
-	virtual ~FConcertLogger();
+	FConcertLogger(const FConcertEndpointContext& InOwnerContext, FLogListener LogListenerFunc);
+	virtual ~FConcertLogger() override;
 
 	// IConcertTransportLogger interface
 	virtual bool IsLogging() const override;
@@ -118,10 +59,13 @@ private:
 
 	/** */
 	bool bIsLogging;
-
+	
 	/** */
 	FConcertEndpointContext OwnerContext;
 
+	/** Called after a log is processed */
+	FLogListener LogListenerFunc;
+
 	/** Queue for unprocessed logs */
 	TQueue<FConcertLog, EQueueMode::Mpsc> LogQueue;
 
diff --git a/Engine/Plugins/Developer/Concert/ConcertMain/Source/Concert/Private/ConcertServer.cpp b/Engine/Plugins/Developer/Concert/ConcertMain/Source/Concert/Private/ConcertServer.cpp
--- a/Engine/Plugins/Developer/Concert/ConcertMain/Source/Concert/Private/ConcertServer.cpp
+++ b/Engine/Plugins/Developer/Concert/ConcertMain/Source/Concert/Private/ConcertServer.cpp
@@ -8,6 +8,7 @@
 #include "ConcertServerSession.h"
 #include "ConcertServerSessionRepositories.h"
 #include "ConcertLogGlobal.h"
+#include "ConcertTransportEvents.h"
 #include "IConcertServerEventSink.h"
 
 #include "Algo/AnyOf.h"
@@ -160,7 +161,13 @@ void FConcertServer::Startup()
 	if (!ServerAdminEndpoint.IsValid() && EndpointProvider.IsValid())
 	{
 		// Create the server administration endpoint
-		ServerAdminEndpoint = EndpointProvider->CreateLocalEndpoint(TEXT("Admin"), Settings->EndpointSettings, &FConcertLogger::CreateLogger);
+		ServerAdminEndpoint = EndpointProvider->CreateLocalEndpoint(TEXT("Admin"), Settings->EndpointSettings, [this](const FConcertEndpointContext& Context)
+		{
+			return FConcertLogger::CreateLogger(Context, [this](const FConcertLog& Log)
+			{
+				ConcertTransportEvents::OnConcertServerLogEvent().Broadcast(*this, Log);
+			});
+		});
 		ServerInfo.AdminEndpointId = ServerAdminEndpoint->GetEndpointContext().EndpointId;
 
 		// Make it discoverable
@@ -1511,7 +1518,13 @@ TSharedPtr<IConcertServerSession> FConcertServer::CreateLiveSession(const FConce
 	TSharedPtr<FConcertServerSession> LiveSession = MakeShared<FConcertServerSession>(
 		LiveSessionInfo,
 		Settings->ServerSettings,
-		EndpointProvider->CreateLocalEndpoint(LiveSessionInfo.SessionName, Settings->EndpointSettings, &FConcertLogger::CreateLogger),
+		EndpointProvider->CreateLocalEndpoint(LiveSessionInfo.SessionName, Settings->EndpointSettings, [this](const FConcertEndpointContext& Context)
+		{
+			return FConcertLogger::CreateLogger(Context, [this](const FConcertLog& Log)
+			{
+				ConcertTransportEvents::OnConcertServerLogEvent().Broadcast(*this, Log);
+			});
+		}),
 		InRepository.GetSessionWorkingDir(LiveSessionInfo.SessionId)
 		);
 
diff --git a/Engine/Plugins/Developer/Concert/ConcertMain/Source/Concert/Private/ConcertTransportEvents.cpp b/Engine/Plugins/Developer/Concert/ConcertMain/Source/Concert/Private/ConcertTransportEvents.cpp
new file mode 100644
--- /dev/null
+++ b/Engine/Plugins/Developer/Concert/ConcertMain/Source/Concert/Private/ConcertTransportEvents.cpp
@@ -0,0 +1,24 @@
+// Copyright Epic Games, Inc. All Rights Reserved.
+
+#include "ConcertTransportEvents.h"
+
+namespace ConcertTransportEvents
+{
+	FConcertTransportLoggingEnabledChanged& OnConcertTransportLoggingEnabledChangedEvent()
+	{
+		static FConcertTransportLoggingEnabledChanged Instance;
+		return Instance;
+	}
+	
+	FConcertClientLogEvent& OnConcertClientLogEvent()
+	{
+		static FConcertClientLogEvent Instance;
+		return Instance;
+	}
+
+	FConcertServerLogEvent& OnConcertServerLogEvent()
+	{
+		static FConcertServerLogEvent Instance;
+		return Instance;
+	}
+}
diff --git a/Engine/Plugins/Developer/Concert/ConcertMain/Source/Concert/Public/ConcertTransportEvents.h b/Engine/Plugins/Developer/Concert/ConcertMain/Source/Concert/Public/ConcertTransportEvents.h
new file mode 100644
--- /dev/null
+++ b/Engine/Plugins/Developer/Concert/ConcertMain/Source/Concert/Public/ConcertTransportEvents.h
@@ -0,0 +1,88 @@
+// Copyright Epic Games, Inc. All Rights Reserved.
+
+#pragma once
+
+#include "CoreMinimal.h"
+#include "ConcertMessageData.h"
+#include "IConcertClient.h"
+#include "ConcertTransportEvents.generated.h"
+
+class IConcertClient;
+class IConcertServer;
+
+UENUM()
+enum class EConcertLogMessageAction : uint8
+{
+	None,
+	Send,
+	Publish,
+	Receive,
+	Queue,
+	Discard,
+	Duplicate,
+	TimeOut,
+	Process,
+	EndpointDiscovery,
+	EndpointTimeOut,
+	EndpointClosure,
+};
+
+USTRUCT()
+struct CONCERT_API FConcertLog
+{
+	GENERATED_BODY()
+
+	UPROPERTY()
+	uint64 Frame = 0;
+
+	UPROPERTY()
+	FGuid MessageId;
+
+	UPROPERTY()
+	uint16 MessageOrderIndex  = 0;
+
+	UPROPERTY()
+	uint16 ChannelId = 0;
+
+	UPROPERTY()
+	FDateTime Timestamp = {0};
+
+	UPROPERTY()
+	EConcertLogMessageAction MessageAction = EConcertLogMessageAction::None;
+
+	UPROPERTY()
+	FName MessageTypeName;
+
+	UPROPERTY()
+	FGuid OriginEndpointId;
+
+	UPROPERTY()
+	FGuid DestinationEndpointId;
+
+	UPROPERTY()
+	FName CustomPayloadTypename;
+
+	UPROPERTY()
+	int32 CustomPayloadUncompressedByteSize = 0;
+
+	UPROPERTY()
+	FString StringPayload;
+
+	UPROPERTY(Transient)
+	FConcertSessionSerializedPayload SerializedPayload;
+};
+
+namespace ConcertTransportEvents
+{
+	DECLARE_MULTICAST_DELEGATE_OneParam(FConcertTransportLoggingEnabledChanged, bool /*bIsEnabled*/);
+	DECLARE_MULTICAST_DELEGATE_TwoParams(FConcertClientLogEvent, const IConcertClient&, const FConcertLog&);
+	DECLARE_MULTICAST_DELEGATE_TwoParams(FConcertServerLogEvent, const IConcertServer&, const FConcertLog&);
+
+	CONCERT_API bool IsLoggingEnabled();
+	CONCERT_API void SetLoggingEnabled(bool bEnabled);
+
+	CONCERT_API FConcertTransportLoggingEnabledChanged& OnConcertTransportLoggingEnabledChangedEvent();
+	CONCERT_API FConcertClientLogEvent& OnConcertClientLogEvent();
+	CONCERT_API FConcertServerLogEvent& OnConcertServerLogEvent();
+}
+
