Add util function to find connected client info by EndpointId and minor clean up

diff --git a/Engine/Plugins/Developer/Concert/ConcertMain/Source/Concert/Private/ConcertServer.cpp b/Engine/Plugins/Developer/Concert/ConcertMain/Source/Concert/Private/ConcertServer.cpp
--- a/Engine/Plugins/Developer/Concert/ConcertMain/Source/Concert/Private/ConcertServer.cpp
+++ b/Engine/Plugins/Developer/Concert/ConcertMain/Source/Concert/Private/ConcertServer.cpp
@@ -642,16 +642,6 @@
 	return Response.ResponseCode == EConcertResponseCode::Success;
 }
 
-TArray<FConcertSessionClientInfo> FConcertServer::GetSessionClients(const FGuid& SessionId) const
-{
-	TSharedPtr<IConcertServerSession> ServerSession = GetLiveSession(SessionId);
-	if (ServerSession)
-	{
-		return ServerSession->GetSessionClients();
-	}
-	return TArray<FConcertSessionClientInfo>();
-}
-
 TArray<FConcertSessionInfo> FConcertServer::GetLiveSessionInfos() const
 {
 	TArray<FConcertSessionInfo> SessionsInfo;
@@ -1434,7 +1424,7 @@
 	const FConcertAdmin_GetSessionClientsRequest* Message = Context.GetMessage<FConcertAdmin_GetSessionClientsRequest>();
 
 	FConcertAdmin_GetSessionClientsResponse ResponseData;
-	ResponseData.SessionClients = GetSessionClients(Message->SessionId);
+	ResponseData.SessionClients = ConcertUtil::GetSessionClients(*this, Message->SessionId);
 	
 	return FConcertAdmin_GetSessionClientsResponse::AsFuture(MoveTemp(ResponseData));
 }
diff --git a/Engine/Plugins/Developer/Concert/ConcertMain/Source/Concert/Private/ConcertServer.h b/Engine/Plugins/Developer/Concert/ConcertMain/Source/Concert/Private/ConcertServer.h
--- a/Engine/Plugins/Developer/Concert/ConcertMain/Source/Concert/Private/ConcertServer.h
+++ b/Engine/Plugins/Developer/Concert/ConcertMain/Source/Concert/Private/ConcertServer.h
@@ -46,7 +46,6 @@
 	virtual bool ExportSession(const FGuid& SessionId, const FConcertSessionFilter& SessionFilter, const FString& DestDir, bool bAnonymizeData, FText& OutFailureReason) override;
 	virtual bool RenameSession(const FGuid& SessionId, const FString& NewName, FText& OutFailureReason) override;
 	virtual bool DestroySession(const FGuid& SessionId, FText& OutFailureReason) override;
-	virtual TArray<FConcertSessionClientInfo> GetSessionClients(const FGuid& SessionId) const override;
 
 private:
 	/** Returns the root dir where the servers keeps the its internally created repositories (When the caller doesn't provide the paths). */
diff --git a/Engine/Plugins/Developer/Concert/ConcertMain/Source/Concert/Private/ConcertUtil.cpp b/Engine/Plugins/Developer/Concert/ConcertMain/Source/Concert/Private/ConcertUtil.cpp
--- a/Engine/Plugins/Developer/Concert/ConcertMain/Source/Concert/Private/ConcertUtil.cpp
+++ b/Engine/Plugins/Developer/Concert/ConcertMain/Source/Concert/Private/ConcertUtil.cpp
@@ -5,6 +5,8 @@
 #include "HAL/FileManager.h"
 #include "ConcertLogGlobal.h"
 #include "ConcertLogger.h"
+#include "ConcertServerSession.h"
+#include "IConcertServer.h"
 
 bool ConcertUtil::DeleteDirectoryTree(const TCHAR* InDirectoryToDelete, const TCHAR* InMoveToDirBeforeDelete)
 {
@@ -58,3 +60,28 @@
 	FConcertLogger::SetVerboseLogging(bInState);
 }
 
+TArray<FConcertSessionClientInfo> ConcertUtil::GetSessionClients(IConcertServer& Server, const FGuid& SessionId)
+{
+	TSharedPtr<IConcertServerSession> ServerSession = Server.GetLiveSession(SessionId);
+	if (ServerSession)
+	{
+		return ServerSession->GetSessionClients();
+	}
+	return TArray<FConcertSessionClientInfo>();
+}
+
+TOptional<FConcertSessionClientInfo> ConcertUtil::GetConnectedClientInfo(IConcertServer& Server, const FGuid& ClientEndpointId)
+{
+	for (const TSharedPtr<IConcertServerSession>& ServerSession : Server.GetLiveSessions())
+	{
+		for (const FConcertSessionClientInfo& ClientInfo : ServerSession->GetSessionClients())
+		{
+			if (ClientEndpointId == ClientInfo.ClientEndpointId)
+			{
+				return ClientInfo;
+			}
+		}
+	}
+
+	return {};
+}
diff --git a/Engine/Plugins/Developer/Concert/ConcertMain/Source/Concert/Public/ConcertUtil.h b/Engine/Plugins/Developer/Concert/ConcertMain/Source/Concert/Public/ConcertUtil.h
--- a/Engine/Plugins/Developer/Concert/ConcertMain/Source/Concert/Public/ConcertUtil.h
+++ b/Engine/Plugins/Developer/Concert/ConcertMain/Source/Concert/Public/ConcertUtil.h
@@ -1,19 +1,37 @@
-// Copyright Epic Games, Inc. All Rights Reserved.
-
-#pragma once
-
-#include "CoreMinimal.h"
-
-namespace ConcertUtil
-{
-
-/** Delete a directory tree via a move and delete */
-CONCERT_API bool DeleteDirectoryTree(const TCHAR* InDirectoryToDelete, const TCHAR* InMoveToDirBeforeDelete = nullptr);
-
-/** Copy the specified data size from a source archive into a destination archive. */
-CONCERT_API bool Copy(FArchive& DstAr, FArchive& SrcAr, int64 Size);
-
-/** Turn on verbose logging for all loggers (including console loggers). */
-CONCERT_API void SetVerboseLogging(bool bInState);
-
-} // namespace ConcertUtil
+// Copyright Epic Games, Inc. All Rights Reserved.
+
+#pragma once
+
+#include "CoreMinimal.h"
+
+class IConcertServer;
+struct FConcertSessionClientInfo;
+
+namespace ConcertUtil
+{
+
+/** Delete a directory tree via a move and delete */
+CONCERT_API bool DeleteDirectoryTree(const TCHAR* InDirectoryToDelete, const TCHAR* InMoveToDirBeforeDelete = nullptr);
+
+/** Copy the specified data size from a source archive into a destination archive. */
+CONCERT_API bool Copy(FArchive& DstAr, FArchive& SrcAr, int64 Size);
+
+/** Turn on verbose logging for all loggers (including console loggers). */
+CONCERT_API void SetVerboseLogging(bool bInState);
+
+/**
+ * Get the list of clients for a session
+ * @param Server The server to look on
+ * @param SessionId The session ID
+ * @return A list of clients connected to the session
+ */
+CONCERT_API TArray<FConcertSessionClientInfo> GetSessionClients(IConcertServer& Server, const FGuid& SessionId);
+
+/**
+ * Gets a connected client's info.
+ * @param Server The server to look on
+ * @param ClientEndpointId The client's endpoint ID
+ * @return The client's info if found (it is not found if the client is not connected to any session)
+ */
+CONCERT_API TOptional<FConcertSessionClientInfo> GetConnectedClientInfo(IConcertServer& Server, const FGuid& ClientEndpointId);
+} // namespace ConcertUtil
diff --git a/Engine/Plugins/Developer/Concert/ConcertMain/Source/Concert/Public/IConcertServer.h b/Engine/Plugins/Developer/Concert/ConcertMain/Source/Concert/Public/IConcertServer.h
--- a/Engine/Plugins/Developer/Concert/ConcertMain/Source/Concert/Public/IConcertServer.h
+++ b/Engine/Plugins/Developer/Concert/ConcertMain/Source/Concert/Public/IConcertServer.h
@@ -172,11 +172,4 @@
 	 * @return true if the session was found and destroyed
 	 */
 	virtual bool DestroySession(const FGuid& SessionId, FText& OutFailureReason) = 0;
-
-	/**
-	 * Get the list of clients for a session
-	 * @param SessionId The session ID
-	 * @return A list of clients connected to the session
-	 */
-	virtual TArray<FConcertSessionClientInfo> GetSessionClients(const FGuid& SessionId) const = 0;
 };
