Fix typo

diff --git a/Engine/Plugins/Developer/Concert/ConcertMain/Source/Concert/Private/ConcertLogger.cpp b/Engine/Plugins/Developer/Concert/ConcertMain/Source/Concert/Private/ConcertLogger.cpp
--- a/Engine/Plugins/Developer/Concert/ConcertMain/Source/Concert/Private/ConcertLogger.cpp
+++ b/Engine/Plugins/Developer/Concert/ConcertMain/Source/Concert/Private/ConcertLogger.cpp
@@ -216,7 +216,7 @@
 		return Algo::AllOf(Loggers, [](const TWeakPtr<FConcertLogger, ESPMode::ThreadSafe>& Logger)
 		{
 			TSharedPtr<FConcertLogger, ESPMode::ThreadSafe> LoggerPtr = Logger.Pin();
-			return !LoggerPtr.IsValid() && LoggerPtr->IsLogging();
+			return LoggerPtr.IsValid() && LoggerPtr->IsLogging();
 		});
 	}
 
