Extend ConcertUtil

diff --git a/Engine/Plugins/Developer/Concert/ConcertMain/Source/Concert/Private/ConcertUtil.cpp b/Engine/Plugins/Developer/Concert/ConcertMain/Source/Concert/Private/ConcertUtil.cpp
--- a/Engine/Plugins/Developer/Concert/ConcertMain/Source/Concert/Private/ConcertUtil.cpp
+++ b/Engine/Plugins/Developer/Concert/ConcertMain/Source/Concert/Private/ConcertUtil.cpp
@@ -70,18 +70,37 @@
 	return TArray<FConcertSessionClientInfo>();
 }
 
-TOptional<FConcertSessionClientInfo> ConcertUtil::GetConnectedClientInfo(IConcertServer& Server, const FGuid& ClientEndpointId)
+namespace ConcertUtil::Private
 {
-	for (const TSharedPtr<IConcertServerSession>& ServerSession : Server.GetLiveSessions())
+	static TOptional<TPair<FConcertSessionClientInfo, TSharedPtr<IConcertServerSession>>> FindByClient(IConcertServer& Server, const FGuid& ClientEndpointId)
 	{
-		for (const FConcertSessionClientInfo& ClientInfo : ServerSession->GetSessionClients())
+		for (const TSharedPtr<IConcertServerSession>& ServerSession : Server.GetLiveSessions())
 		{
-			if (ClientEndpointId == ClientInfo.ClientEndpointId)
+			for (const FConcertSessionClientInfo& ClientInfo : ServerSession->GetSessionClients())
 			{
-				return ClientInfo;
+				if (ClientEndpointId == ClientInfo.ClientEndpointId)
+				{
+					return {{ ClientInfo, ServerSession }};
+				}
 			}
 		}
+
+		return {};
 	}
+}
 
-	return {};
+TOptional<FConcertSessionClientInfo> ConcertUtil::GetConnectedClientInfo(IConcertServer& Server, const FGuid& ClientEndpointId)
+{
+	const TOptional<TPair<FConcertSessionClientInfo, TSharedPtr<IConcertServerSession>>> Result = Private::FindByClient(Server, ClientEndpointId);
+	return Result
+		? Result->Key
+		: TOptional<FConcertSessionClientInfo>();
 }
+
+TSharedPtr<IConcertServerSession> ConcertUtil::GetLiveSessionClientConnectedTo(IConcertServer& Server, const FGuid& ClientEndpointId)
+{
+	const TOptional<TPair<FConcertSessionClientInfo, TSharedPtr<IConcertServerSession>>> Result = Private::FindByClient(Server, ClientEndpointId);
+	return Result
+		? Result->Value
+		: TSharedPtr<IConcertServerSession>();
+}
diff --git a/Engine/Plugins/Developer/Concert/ConcertMain/Source/Concert/Public/ConcertUtil.h b/Engine/Plugins/Developer/Concert/ConcertMain/Source/Concert/Public/ConcertUtil.h
--- a/Engine/Plugins/Developer/Concert/ConcertMain/Source/Concert/Public/ConcertUtil.h
+++ b/Engine/Plugins/Developer/Concert/ConcertMain/Source/Concert/Public/ConcertUtil.h
@@ -4,6 +4,7 @@
 
 #include "CoreMinimal.h"
 
+class IConcertServerSession;
 class IConcertServer;
 struct FConcertSessionClientInfo;
 
@@ -34,4 +35,12 @@
  * @return The client's info if found (it is not found if the client is not connected to any session)
  */
 CONCERT_API TOptional<FConcertSessionClientInfo> GetConnectedClientInfo(IConcertServer& Server, const FGuid& ClientEndpointId);
+
+/**
+ * Retrieves the server session object from the given connected client.
+ * @param Server The server to look on
+ * @param ClientEndpointId The client's endpoint ID
+ * @return The session the client is connected to if any
+ */
+CONCERT_API TSharedPtr<IConcertServerSession> GetLiveSessionClientConnectedTo(IConcertServer& Server, const FGuid& ClientEndpointId);
 } // namespace ConcertUtil
