Deprecating ANY_PACKAGE.

This change consists of multiple changes:

Core:
- Deprecation of ANY_PACKAGE macro. Added ANY_PACKAGE_DEPRECATED macro which can still be used for backwards compatibility purposes (only used in CoreUObject)
- Deprecation of StaticFindObjectFast* functions that take bAnyPackage parameter
- Added UStruct::GetStructPathName function that returns FTopLevelAssetPath representing the path name (package + object FName, super quick compared to UObject::GetPathName) + wrapper UClass::GetClassPathName to make it look better when used with UClasses
- Added (Static)FindFirstObject* functions that find a first object given its Name (no Outer). These functions are used in places I consider valid to do global UObject (UClass) lookups like parsing command line parameters / checking for unique object names
- Added static UClass::TryFindType function which serves a similar purpose as FindFirstObject however it's going to throw a warning (with a callstack / maybe ensure in the future?) if short class name is provided. This function is used  in places that used to use short class names but now should have been converted to use path names to catch any potential regressions and or edge cases I missed.
- Added static UClass::TryConvertShortNameToPathName utility function
- Added static UClass::TryFixShortClassNameExportPath utility function
- Object text export paths will now also include class path (Texture2D'/Game/Textures/Grass.Grass' -> /Script/Engine.Texture2D'/Game/Textures/Grass.Grass')
- All places that manually generated object export paths for objects will now use FObjectPropertyBase::GetExportPath
- Added a new startup test that checks for short type names in UClass/FProperty MetaData values

AssetRegistry:
- Deprecated any member variables (FAssetData / FARFilter) or functions that use FNames to represent class names and replaced them with FTopLevelAssetPath
- Added new member variables and new function overloads that use FTopLevelAssetPath to represent class names
- This also applies to a few other modules' APIs to match AssetRegistry changes

Everything else:
- Updated code that used ANY_PACKAGE (depending on the use case) to use FindObject(nullptr, PathToObject), UClass::TryFindType (used when path name is expected, warns if it's a short name) or FindFirstObject (usually for finding types based on user input but there's been a few legitimate use cases not related to user input)
- Updated code that used AssetRegistry API to use FTopLevelAssetPaths and USomeClass::StaticClass()->GetClassPathName() instead of GetFName()
- Updated meta data and hardcoded FindObject(ANY_PACKAGE, "EEnumNameOrClassName") calls to use path names

diff --git a/Engine/Config/BaseEditorPerProjectUserSettings.ini b/Engine/Config/BaseEditorPerProjectUserSettings.ini
--- a/Engine/Config/BaseEditorPerProjectUserSettings.ini
+++ b/Engine/Config/BaseEditorPerProjectUserSettings.ini
@@ -774,43 +774,43 @@ DefaultBrushRadius=128
 
 [BlueprintSpawnNodes]
 ; Comment box is bound to C, but that is handled differently due to it needing to work without clicking
-+Node=(Class=Actor:ReceiveBeginPlay Key=P Shift=false Ctrl=false Alt=false)
++Node=(Class=/Script/Engine.Actor:ReceiveBeginPlay Key=P Shift=false Ctrl=false Alt=false)
 +Node=(Class="Do N" Key=N Shift=false Ctrl=false Alt=false)
-+Node=(Class=KismetSystemLibrary:Delay Key=D Shift=false Ctrl=false Alt=false)
-+Node=(Class=K2Node_IfThenElse Key=B Shift=false Ctrl=false Alt=false)
-+Node=(Class=K2Node_ExecutionSequence Key=S Shift=false Ctrl=false Alt=false)
++Node=(Class=/Script/Engine.KismetSystemLibrary:Delay Key=D Shift=false Ctrl=false Alt=false)
++Node=(Class=/Script/BlueprintGraph.K2Node_IfThenElse Key=B Shift=false Ctrl=false Alt=false)
++Node=(Class=/Script/BlueprintGraph.K2Node_ExecutionSequence Key=S Shift=false Ctrl=false Alt=false)
 +Node=(Class=Gate Key=G Shift=false Ctrl=false Alt=false)
-+Node=(Class=K2Node_MultiGate Key=M Shift=false Ctrl=false Alt=false)
++Node=(Class=/Script/BlueprintGraph.K2Node_MultiGate Key=M Shift=false Ctrl=false Alt=false)
 +Node=(Class=ForEachLoop Key=F Shift=false Ctrl=false Alt=false)
 +Node=(Class=DoOnce Key=O Shift=false Ctrl=false Alt=false)
-+Node=(Class=KismetArrayLibrary:Array_Get Key=A Shift=false Ctrl=false Alt=false)
++Node=(Class=/Script/Engine.KismetArrayLibrary:Array_Get Key=A Shift=false Ctrl=false Alt=false)
 
 [DefaultEventNodes]
-+Node=(TargetClass=Actor TargetEvent="ReceiveBeginPlay")
-+Node=(TargetClass=Actor TargetEvent="ReceiveActorBeginOverlap")
-+Node=(TargetClass=Actor TargetEvent="ReceiveTick")
-+Node=(TargetClass=ActorComponent TargetEvent="ReceiveBeginPlay")
-+Node=(TargetClass=ActorComponent TargetEvent="ReceiveTick")
-+Node=(TargetClass=GameplayAbility TargetEvent="K2_ActivateAbility")
-+Node=(TargetClass=GameplayAbility TargetEvent="K2_OnEndAbility")
-+Node=(TargetClass=GameplayCueNotify_BurstLatent TargetEvent="OnBurst")
-+Node=(TargetClass=GameplayCueNotify_Looping TargetEvent="OnApplication")
-+Node=(TargetClass=GameplayCueNotify_Looping TargetEvent="OnLoopingStart")
-+Node=(TargetClass=GameplayCueNotify_Looping TargetEvent="OnRecurring")
-+Node=(TargetClass=GameplayCueNotify_Looping TargetEvent="OnRemoval")
-+Node=(TargetClass=UserWidget TargetEvent="PreConstruct")
-+Node=(TargetClass=UserWidget TargetEvent="Construct")
-+Node=(TargetClass=UserWidget TargetEvent="Tick")
-+Node=(TargetClass=AnimInstance TargetEvent="BlueprintUpdateAnimation")
-+Node=(TargetClass=FunctionalTest TargetEvent="ReceivePrepareTest")
-+Node=(TargetClass=FunctionalTest TargetEvent="ReceiveStartTest")
-+Node=(TargetClass=FunctionalTest TargetEvent="ReceiveTick")
-+Node=(TargetClass=UserDefinedCaptureProtocol TargetEvent="OnPreTick")
-+Node=(TargetClass=UserDefinedCaptureProtocol TargetEvent="OnTick")
-+Node=(TargetClass=UserDefinedCaptureProtocol TargetEvent="OnStartCapture")
-+Node=(TargetClass=UserDefinedCaptureProtocol TargetEvent="OnCaptureFrame")
-+Node=(TargetClass=UserDefinedCaptureProtocol TargetEvent="OnBeginFinalize")
-+Node=(TargetClass=UserDefinedCaptureProtocol TargetEvent="OnFinalize")
++Node=(TargetClass=/Script/Engine.Actor TargetEvent="ReceiveBeginPlay")
++Node=(TargetClass=/Script/Engine.Actor TargetEvent="ReceiveActorBeginOverlap")
++Node=(TargetClass=/Script/Engine.Actor TargetEvent="ReceiveTick")
++Node=(TargetClass=/Script/Engine.ActorComponent TargetEvent="ReceiveBeginPlay")
++Node=(TargetClass=/Script/Engine.ActorComponent TargetEvent="ReceiveTick")
++Node=(TargetClass=/Script/GameplayAbilities.GameplayAbility TargetEvent="K2_ActivateAbility")
++Node=(TargetClass=/Script/GameplayAbilities.GameplayAbility TargetEvent="K2_OnEndAbility")
++Node=(TargetClass=/Script/GameplayAbilities.GameplayCueNotify_BurstLatent TargetEvent="OnBurst")
++Node=(TargetClass=/Script/GameplayAbilities.GameplayCueNotify_Looping TargetEvent="OnApplication")
++Node=(TargetClass=/Script/GameplayAbilities.GameplayCueNotify_Looping TargetEvent="OnLoopingStart")
++Node=(TargetClass=/Script/GameplayAbilities.GameplayCueNotify_Looping TargetEvent="OnRecurring")
++Node=(TargetClass=/Script/GameplayAbilities.GameplayCueNotify_Looping TargetEvent="OnRemoval")
++Node=(TargetClass=/Script/UMG.UserWidget TargetEvent="PreConstruct")
++Node=(TargetClass=/Script/UMG.UserWidget TargetEvent="Construct")
++Node=(TargetClass=/Script/UMG.UserWidget TargetEvent="Tick")
++Node=(TargetClass=/Script/Engine.AnimInstance TargetEvent="BlueprintUpdateAnimation")
++Node=(TargetClass=/Script/FunctionalTesting.FunctionalTest TargetEvent="ReceivePrepareTest")
++Node=(TargetClass=/Script/FunctionalTesting.FunctionalTest TargetEvent="ReceiveStartTest")
++Node=(TargetClass=/Script/FunctionalTesting.FunctionalTest TargetEvent="ReceiveTick")
++Node=(TargetClass=/Script/MovieSceneCapture.UserDefinedCaptureProtocol TargetEvent="OnPreTick")
++Node=(TargetClass=/Script/MovieSceneCapture.UserDefinedCaptureProtocol TargetEvent="OnTick")
++Node=(TargetClass=/Script/MovieSceneCapture.UserDefinedCaptureProtocol TargetEvent="OnStartCapture")
++Node=(TargetClass=/Script/MovieSceneCapture.UserDefinedCaptureProtocol TargetEvent="OnCaptureFrame")
++Node=(TargetClass=/Script/MovieSceneCapture.UserDefinedCaptureProtocol TargetEvent="OnBeginFinalize")
++Node=(TargetClass=/Script/MovieSceneCapture.UserDefinedCaptureProtocol TargetEvent="OnFinalize")
 
 [MaterialEditorSpawnNodes]
 +Node=(Class=MaterialExpressionAdd Key=A Shift=false Ctrl=false Alt=false)
@@ -940,7 +940,6 @@ bShowTransactionDetails=false
 [/Script/WindowsMixedRealityRuntimeSettings.WindowsMixedRealityRuntimeSettings]
 RemoteHoloLensIP=
 MaxBitrate=4000
-
 [/Script/LiveCoding.LiveCodingSettings]
 Startup=AutomaticButHidden
 bEnabled=False
diff --git a/Engine/Config/BaseEngine.ini b/Engine/Config/BaseEngine.ini
--- a/Engine/Config/BaseEngine.ini
+++ b/Engine/Config/BaseEngine.ini
@@ -2340,20 +2340,20 @@ AspectRatioAxisConstraint=AspectRatio_MaintainXFOV
 ; Fill out this list with asset registry tags that are not needed in cooked builds and will be stripped during cook.
 ; You can use a * wildcard to refer to either all classes or all tags for the class.
 ; If your game does not need any asset registry tags, simply add (Class=*,Tag=*). Be warned that this this may make some engine systems fail to work since they might need ParentClass or GeneratedClass on Blueprints.
-+CookedTagsBlacklist=(Class=Blueprint,Tag=FiB)
-+CookedTagsBlacklist=(Class=Blueprint,Tag=FiBData)
++CookedTagsBlacklist=(Class=/Script/Engine.Blueprint,Tag=FiB)
++CookedTagsBlacklist=(Class=/Script/Engine.Blueprint,Tag=FiBData)
 +CookedTagsBlacklist=(Class=*,Tag=AssetImportData)
 
 ; Alternatively, you can use an allow list which will exclude all tags except those mentioned in this list. Set bUseAssetRegistryTagsWhitelistInsteadOfBlacklist to true to use this list.
 bUseAssetRegistryTagsWhitelistInsteadOfBlacklist=false
-+CookedTagsWhitelist=(Class=Blueprint,Tag=ParentClass)
-+CookedTagsWhitelist=(Class=Blueprint,Tag=GeneratedClass)
-+CookedTagsWhitelist=(Class=Blueprint,Tag=GameplayCueName)
++CookedTagsWhitelist=(Class=/Script/Engine.Blueprint,Tag=ParentClass)
++CookedTagsWhitelist=(Class=/Script/Engine.Blueprint,Tag=GeneratedClass)
++CookedTagsWhitelist=(Class=/Script/Engine.Blueprint,Tag=GameplayCueName)
 +CookedTagsWhitelist=(Class=*,Tag=AssetBundleData)
 +CookedTagsWhitelist=(Class=*,Tag=PrimaryAssetType)
 +CookedTagsWhitelist=(Class=*,Tag=PrimaryAssetName)
-+CookedTagsWhitelist=(Class=World,Tag=Tests)
-+CookedTagsWhitelist=(Class=World,Tag=TestNames)
++CookedTagsWhitelist=(Class=/Script/Engine.World,Tag=Tests)
++CookedTagsWhitelist=(Class=/Script/Engine.World,Tag=TestNames)
 
 ; Configure tag values to be cooked and loaded as FName / FRegistryExportPath
 +CookedTagsAsFName=PrimaryAssetType
diff --git a/Engine/Config/BaseGame.ini b/Engine/Config/BaseGame.ini
--- a/Engine/Config/BaseGame.ini
+++ b/Engine/Config/BaseGame.ini
@@ -228,9 +228,9 @@ FrameInterpolationProcessor=Class'/Script/LiveLink.LiveLinkBasicFrameInterpolati
 [CookedEditorSettings]
 +EngineAssetPaths="/Engine"
 +ProjectAssetPaths="/Game"
-+DisallowedObjectClassesToLoad="Blueprint"
-+DisallowedObjectClassesToLoad="AssetImportData"
-+DisallowedAssetClassesToGather="World"
++DisallowedObjectClassesToLoad="/Script/Engine.Blueprint"
++DisallowedObjectClassesToLoad="/Script/Engine.AssetImportData"
++DisallowedAssetClassesToGather="/Script/Engine.World"
 +DisallowedPathsToGather="/Game/Developers"
 ;+EngineSpecificAssetsToCook="/Engine/Maps/Templates/Template_Default.umap"
 
diff --git a/Engine/Plugins/2D/Paper2D/Source/Paper2DEditor/Private/ContentBrowserExtensions/ContentBrowserExtensions.cpp b/Engine/Plugins/2D/Paper2D/Source/Paper2DEditor/Private/ContentBrowserExtensions/ContentBrowserExtensions.cpp
--- a/Engine/Plugins/2D/Paper2D/Source/Paper2DEditor/Private/ContentBrowserExtensions/ContentBrowserExtensions.cpp
+++ b/Engine/Plugins/2D/Paper2D/Source/Paper2DEditor/Private/ContentBrowserExtensions/ContentBrowserExtensions.cpp
@@ -287,7 +287,7 @@ public:
 		for (auto AssetIt = SelectedAssets.CreateConstIterator(); AssetIt; ++AssetIt)
 		{
 			const FAssetData& Asset = *AssetIt;
-			bAnyTextures = bAnyTextures || (Asset.AssetClass == UTexture2D::StaticClass()->GetFName());
+			bAnyTextures = bAnyTextures || (Asset.AssetClassPath == UTexture2D::StaticClass()->GetClassPathName());
 		}
 
 		if (bAnyTextures)
diff --git a/Engine/Plugins/2D/Paper2D/Source/Paper2DEditor/Private/SpriteEditor/SSpriteList.cpp b/Engine/Plugins/2D/Paper2D/Source/Paper2DEditor/Private/SpriteEditor/SSpriteList.cpp
--- a/Engine/Plugins/2D/Paper2D/Source/Paper2DEditor/Private/SpriteEditor/SSpriteList.cpp
+++ b/Engine/Plugins/2D/Paper2D/Source/Paper2DEditor/Private/SpriteEditor/SSpriteList.cpp
@@ -26,7 +26,7 @@ void SSpriteList::RebuildWidget(UTexture2D* NewTextureFilter)
 	FAssetPickerConfig Config;
 
 	//Config.Filter.bRecursiveClasses = true;
-	Config.Filter.ClassNames.Add(UPaperSprite::StaticClass()->GetFName());
+	Config.Filter.ClassPaths.Add(UPaperSprite::StaticClass()->GetClassPathName());
 
 	Config.InitialAssetViewType = EAssetViewType::Tile;
 	Config.SyncToAssetsDelegates.Add(&SyncToAssetsDelegate);
diff --git a/Engine/Plugins/2D/Paper2D/Source/Paper2DEditor/Private/SpriteEditor/SpriteEditorViewportClient.cpp b/Engine/Plugins/2D/Paper2D/Source/Paper2DEditor/Private/SpriteEditor/SpriteEditorViewportClient.cpp
--- a/Engine/Plugins/2D/Paper2D/Source/Paper2DEditor/Private/SpriteEditor/SpriteEditorViewportClient.cpp
+++ b/Engine/Plugins/2D/Paper2D/Source/Paper2DEditor/Private/SpriteEditor/SpriteEditorViewportClient.cpp
@@ -692,7 +692,7 @@ void FSpriteEditorViewportClient::UpdateRelatedSpritesList()
 	if (Texture != nullptr)
 	{
 		FARFilter Filter;
-		Filter.ClassNames.Add(UPaperSprite::StaticClass()->GetFName());
+		Filter.ClassPaths.Add(UPaperSprite::StaticClass()->GetClassPathName());
 		const FName SourceTexturePropName(TEXT("SourceTexture"));
 		Filter.TagsAndValues.Add(SourceTexturePropName, TSoftObjectPtr<UTexture2D>(Texture).ToString());
 		// Legacy format for files that haven't been resaved
diff --git a/Engine/Plugins/AI/EnvironmentQueryEditor/Source/EnvironmentQueryEditor/Private/SEnvQueryProfiler.cpp b/Engine/Plugins/AI/EnvironmentQueryEditor/Source/EnvironmentQueryEditor/Private/SEnvQueryProfiler.cpp
--- a/Engine/Plugins/AI/EnvironmentQueryEditor/Source/EnvironmentQueryEditor/Private/SEnvQueryProfiler.cpp
+++ b/Engine/Plugins/AI/EnvironmentQueryEditor/Source/EnvironmentQueryEditor/Private/SEnvQueryProfiler.cpp
@@ -207,12 +207,12 @@ static UObject* FindQueryObjectByName(FName StatName)
 		AssetPathName.LeftInline(SepIdx);
 	}
 
-	UObject* QueryOb = FindObject<UObject>(ANY_PACKAGE, *AssetPathName);
+	UObject* QueryOb = FindFirstObject<UObject>(*AssetPathName, EFindFirstObjectOptions::NativeFirst);
 	if (QueryOb == nullptr)
 	{
 		TArray<FAssetData> Assets;
 		IAssetRegistry& AssetRegistry = FModuleManager::LoadModuleChecked<FAssetRegistryModule>(TEXT("AssetRegistry")).Get();
-		AssetRegistry.GetAssetsByClass(UEnvQuery::StaticClass()->GetFName(), Assets);
+		AssetRegistry.GetAssetsByClass(UEnvQuery::StaticClass()->GetClassPathName(), Assets);
 
 		const FString MatchObjectName = FString(".") + AssetPathName;
 		for (int32 Idx = 0; Idx < Assets.Num(); Idx++)
diff --git a/Engine/Plugins/AI/MLAdapter/Source/MLAdapter/Private/Agents/MLAdapterAgent.cpp b/Engine/Plugins/AI/MLAdapter/Source/MLAdapter/Private/Agents/MLAdapterAgent.cpp
--- a/Engine/Plugins/AI/MLAdapter/Source/MLAdapter/Private/Agents/MLAdapterAgent.cpp
+++ b/Engine/Plugins/AI/MLAdapter/Source/MLAdapter/Private/Agents/MLAdapterAgent.cpp
@@ -276,7 +276,7 @@ void UMLAdapterAgent::Configure(const FMLAdapterAgentConfig& NewConfig)
 
 	if (NewConfig.AvatarClassName != NAME_None)
 	{
-		AvatarClass = FindObject<UClass>(ANY_PACKAGE, *NewConfig.AvatarClassName.ToString());
+		AvatarClass = UClass::TryFindTypeSlow<UClass>(*NewConfig.AvatarClassName.ToString());
 	}
 
 	if (!AvatarClass)
diff --git a/Engine/Plugins/AI/MLAdapter/Source/MLAdapter/Private/MLAdapterSpace.cpp b/Engine/Plugins/AI/MLAdapter/Source/MLAdapter/Private/MLAdapterSpace.cpp
--- a/Engine/Plugins/AI/MLAdapter/Source/MLAdapter/Private/MLAdapterSpace.cpp
+++ b/Engine/Plugins/AI/MLAdapter/Source/MLAdapter/Private/MLAdapterSpace.cpp
@@ -5,7 +5,7 @@
 #include "MLAdapterJson.h"
 
 
-DEFINE_ENUM_TO_STRING(EMLAdapterSpaceType);
+DEFINE_ENUM_TO_STRING(EMLAdapterSpaceType, "/Script/MLAdapter");
 
 namespace FMLAdapter
 {
diff --git a/Engine/Plugins/AI/MLAdapter/Source/MLAdapter/Public/MLAdapterTypes.h b/Engine/Plugins/AI/MLAdapter/Source/MLAdapter/Public/MLAdapterTypes.h
--- a/Engine/Plugins/AI/MLAdapter/Source/MLAdapter/Public/MLAdapterTypes.h
+++ b/Engine/Plugins/AI/MLAdapter/Source/MLAdapter/Public/MLAdapterTypes.h
@@ -14,9 +14,9 @@ class AController;
 DECLARE_LOG_CATEGORY_EXTERN(LogMLAdapter, Log, All);
 
 #ifndef DEFINE_ENUM_TO_STRING
-#define DEFINE_ENUM_TO_STRING(EnumType) FString EnumToString(const EnumType Value) \
+#define DEFINE_ENUM_TO_STRING(EnumType, EnumPackage) FString EnumToString(const EnumType Value) \
 { \
-	static const UEnum* TypeEnum = FindObject<UEnum>(ANY_PACKAGE, TEXT(#EnumType)); \
+	static const UEnum* TypeEnum = FindObject<UEnum>(nullptr, TEXT(EnumPackage) TEXT(".") TEXT(#EnumType)); \
 	return TypeEnum->GetNameStringByIndex(static_cast<int32>(Value)); \
 }
 #endif //DEFINE_ENUM_TO_STRING
diff --git a/Engine/Plugins/Animation/ControlRig/Source/ControlRigDeveloper/Private/ControlRigVisualGraphUtils.cpp b/Engine/Plugins/Animation/ControlRig/Source/ControlRigDeveloper/Private/ControlRigVisualGraphUtils.cpp
--- a/Engine/Plugins/Animation/ControlRig/Source/ControlRigDeveloper/Private/ControlRigVisualGraphUtils.cpp
+++ b/Engine/Plugins/Animation/ControlRig/Source/ControlRigDeveloper/Private/ControlRigVisualGraphUtils.cpp
@@ -64,7 +64,7 @@ FAutoConsoleCommandWithWorldAndArgs FCmdControlRigVisualGraphUtilsDumpHierarchy
 		TArray<URigHierarchy*> Hierarchies;
 		for(const FString& ObjectPathName : ObjectPathNames)
 		{
-			if(UObject* Object = FindObject<UObject>(ANY_PACKAGE, *ObjectPathName, false))
+			if(UObject* Object = FindObject<UObject>(nullptr, *ObjectPathName, false))
 			{
 				if(UControlRig* CR = Cast<UControlRig>(Object))
 				{
@@ -76,7 +76,7 @@ FAutoConsoleCommandWithWorldAndArgs FCmdControlRigVisualGraphUtilsDumpHierarchy
 				}
 				else
 				{
-					UE_LOG(LogVisualGraphUtils, Error, TEXT("Object is not a hierarchy / nor a Control Rig: '%s'"), *ObjectPathName);
+					UE_LOG(LogVisualGraphUtils, Error, TEXT("Object is not a hierarchy / nor a Control Rig / or short name was provided: '%s'"), *ObjectPathName);
 					return;
 				}
 			}
diff --git a/Engine/Plugins/Animation/ControlRig/Source/ControlRigEditor/Private/ControlRigEditorModule.cpp b/Engine/Plugins/Animation/ControlRig/Source/ControlRigEditor/Private/ControlRigEditorModule.cpp
--- a/Engine/Plugins/Animation/ControlRig/Source/ControlRigEditor/Private/ControlRigEditorModule.cpp
+++ b/Engine/Plugins/Animation/ControlRig/Source/ControlRigEditor/Private/ControlRigEditorModule.cpp
@@ -542,7 +542,7 @@ TSharedRef< SWidget > FControlRigEditorModule::GenerateAnimationMenu(TWeakPtr<IA
 										const bool bMatchesFlags = !InUnloadedClassData->HasAnyClassFlags(CLASS_Hidden | CLASS_HideDropDown | CLASS_Deprecated | CLASS_Abstract);
 										if (bChildOfObjectClass && bMatchesFlags)
 										{
-											FString GeneratedClassPathString = InUnloadedClassData->GetClassPath().ToString();
+											FString GeneratedClassPathString = InUnloadedClassData->GetClassPathName().ToString();
 											FName BlueprintPath = FName(*GeneratedClassPathString.LeftChop(2)); // Chop off _C
 											FAssetData AssetData = AssetRegistry.GetAssetByObjectPath(BlueprintPath);
 											return MatchesFilter(AssetData);
@@ -1128,7 +1128,7 @@ void FControlRigEditorModule::GetTypeActions(UControlRigBlueprint* CRB, FBluepri
 		// find all control rigs in the project
 		TArray<FAssetData> ControlRigAssetDatas;
 		FARFilter ControlRigAssetFilter;
-		ControlRigAssetFilter.ClassNames.Add(UControlRigBlueprint::StaticClass()->GetFName());
+		ControlRigAssetFilter.ClassPaths.Add(UControlRigBlueprint::StaticClass()->GetClassPathName());
 		AssetRegistryModule.Get().GetAssets(ControlRigAssetFilter, ControlRigAssetDatas);
 
 		// loop over all control rigs in the project
diff --git a/Engine/Plugins/Animation/ControlRig/Source/ControlRigEditor/Private/EditMode/SControlRigBaseListWidget.cpp b/Engine/Plugins/Animation/ControlRig/Source/ControlRigEditor/Private/EditMode/SControlRigBaseListWidget.cpp
--- a/Engine/Plugins/Animation/ControlRig/Source/ControlRigEditor/Private/EditMode/SControlRigBaseListWidget.cpp
+++ b/Engine/Plugins/Animation/ControlRig/Source/ControlRigEditor/Private/EditMode/SControlRigBaseListWidget.cpp
@@ -718,7 +718,7 @@ void SControlRigBaseListWidget::Construct(const FArguments& InArgs)
 	FAssetPickerConfig AssetPickerConfig;
 
 	//AssetPickerConfig.Filter.bRecursiveClasses = true;
-	AssetPickerConfig.Filter.ClassNames.Add(UControlRigPoseAsset::StaticClass()->GetFName());
+	AssetPickerConfig.Filter.ClassPaths.Add(UControlRigPoseAsset::StaticClass()->GetClassPathName());
 
 	AssetPickerConfig.InitialAssetViewType = EAssetViewType::Tile;
 	AssetPickerConfig.bAllowDragging = false;
@@ -895,7 +895,7 @@ void SControlRigBaseListWidget::UpdateInputValidity()
 	const bool bAllowExistingAsset = true;
 
 	FName AssetClassName = AssetClassNames.Num() == 1 ? AssetClassNames[0] : NAME_None;
-	UClass* AssetClass = AssetClassName != NAME_None ? FindObject<UClass>(ANY_PACKAGE, *AssetClassName.ToString(), true) : nullptr;
+	UClass* AssetClass = AssetClassName != NAME_None ? FindObject<UClass>(nullptr, *AssetClassName.ToString(), true) : nullptr;
 
 	if (!ContentBrowserUtils::IsValidObjectPathForCreate(ObjectPath, AssetClass, ErrorMessage, bAllowExistingAsset))
 	{
@@ -989,7 +989,7 @@ void SControlRigBaseListWidget::FilterChanged()
 
 	FARFilter NewFilter;
 	NewFilter.PackagePaths.Add(CurrentlySelectedVirtualPath);
-	NewFilter.ClassNames.Add(UControlRigPoseAsset::StaticClass()->GetFName());
+	NewFilter.ClassPaths.Add(UControlRigPoseAsset::StaticClass()->GetClassPathName());
 	SetFilterDelegate.ExecuteIfBound(NewFilter);
 }
 
diff --git a/Engine/Plugins/Animation/ControlRig/Source/ControlRigEditor/Private/Editor/ControlRigEditor.cpp b/Engine/Plugins/Animation/ControlRig/Source/ControlRigEditor/Private/Editor/ControlRigEditor.cpp
--- a/Engine/Plugins/Animation/ControlRig/Source/ControlRigEditor/Private/Editor/ControlRigEditor.cpp
+++ b/Engine/Plugins/Animation/ControlRig/Source/ControlRigEditor/Private/Editor/ControlRigEditor.cpp
@@ -604,7 +604,7 @@ void FControlRigEditor::InitFunctionReferences()
 				// find all control rigs in the project
 				TArray<FAssetData> ControlRigAssetDatas;
 				FARFilter ControlRigAssetFilter;
-				ControlRigAssetFilter.ClassNames.Add(UControlRigBlueprint::StaticClass()->GetFName());
+				ControlRigAssetFilter.ClassPaths.Add(UControlRigBlueprint::StaticClass()->GetClassPathName());
 				AssetRegistryModule.Get().GetAssets(ControlRigAssetFilter, ControlRigAssetDatas);
 
 				// loop over all control rigs in the project
diff --git a/Engine/Plugins/Animation/ControlRig/Source/ControlRigEditor/Private/Editor/SControlRigFunctionBulkEditWidget.cpp b/Engine/Plugins/Animation/ControlRig/Source/ControlRigEditor/Private/Editor/SControlRigFunctionBulkEditWidget.cpp
--- a/Engine/Plugins/Animation/ControlRig/Source/ControlRigEditor/Private/Editor/SControlRigFunctionBulkEditWidget.cpp
+++ b/Engine/Plugins/Animation/ControlRig/Source/ControlRigEditor/Private/Editor/SControlRigFunctionBulkEditWidget.cpp
@@ -103,7 +103,7 @@ TSharedRef<SWidget> SControlRigFunctionBulkEditWidget::MakeAssetViewForReference
 	AssetPickerConfig.bShowBottomToolbar = false;
 	AssetPickerConfig.bAutohideSearchBar = true;
 	AssetPickerConfig.bPreloadAssetsForContextMenu = false;
-	AssetPickerConfig.Filter.ClassNames.Add(TEXT("InvalidClass"));
+	AssetPickerConfig.Filter.ClassPaths.Add(FTopLevelAssetPath(TEXT("/Script/Engine"), TEXT("InvalidClass")));
 	AssetPickerConfig.Filter.bIncludeOnlyOnDiskAssets = false;
 	AssetPickerConfig.OnAssetsActivated = FOnAssetsActivated::CreateSP(this, &SControlRigFunctionBulkEditWidget::OnAssetsActivated);
 	AssetPickerConfig.OnGetAssetContextMenu = FOnGetAssetContextMenu::CreateSP( this, &SControlRigFunctionBulkEditWidget::OnGetAssetContextMenu );
diff --git a/Engine/Plugins/Animation/ControlRig/Source/ControlRigEditor/Private/Editor/SRigCurveContainer.cpp b/Engine/Plugins/Animation/ControlRig/Source/ControlRigEditor/Private/Editor/SRigCurveContainer.cpp
--- a/Engine/Plugins/Animation/ControlRig/Source/ControlRigEditor/Private/Editor/SRigCurveContainer.cpp
+++ b/Engine/Plugins/Animation/ControlRig/Source/ControlRigEditor/Private/Editor/SRigCurveContainer.cpp
@@ -652,8 +652,8 @@ void SRigCurveContainer::CreateImportMenu(FMenuBuilder& MenuBuilder)
 
 bool SRigCurveContainer::ShouldFilterOnImport(const FAssetData& AssetData) const
 {
-	return (AssetData.AssetClass != USkeletalMesh::StaticClass()->GetFName() &&
-		AssetData.AssetClass != USkeleton::StaticClass()->GetFName());
+	return (AssetData.AssetClassPath != USkeletalMesh::StaticClass()->GetClassPathName() &&
+		AssetData.AssetClassPath != USkeleton::StaticClass()->GetClassPathName());
 }
 
 void SRigCurveContainer::ImportCurve(const FAssetData& InAssetData)
diff --git a/Engine/Plugins/Animation/ControlRig/Source/ControlRigEditor/Private/Sequencer/ControlRigParameterTrackEditor.cpp b/Engine/Plugins/Animation/ControlRig/Source/ControlRigEditor/Private/Sequencer/ControlRigParameterTrackEditor.cpp
--- a/Engine/Plugins/Animation/ControlRig/Source/ControlRigEditor/Private/Sequencer/ControlRigParameterTrackEditor.cpp
+++ b/Engine/Plugins/Animation/ControlRig/Source/ControlRigEditor/Private/Sequencer/ControlRigParameterTrackEditor.cpp
@@ -693,7 +693,7 @@ public:
 		const bool bMatchesFlags = !InUnloadedClassData->HasAnyClassFlags(CLASS_Hidden | CLASS_HideDropDown | CLASS_Deprecated | CLASS_Abstract);
 		if (bChildOfObjectClass && bMatchesFlags)
 		{
-			FString GeneratedClassPathString = InUnloadedClassData->GetClassPath().ToString();
+			FString GeneratedClassPathString = InUnloadedClassData->GetClassPathName().ToString();
 			FName BlueprintPath = FName(*GeneratedClassPathString.LeftChop(2)); // Chop off _C
 			FAssetData AssetData = AssetRegistry.GetAssetByObjectPath(BlueprintPath);
 			return MatchesFilter(AssetData);
@@ -1251,7 +1251,7 @@ void FControlRigParameterTrackEditor::HandleAddControlRigSubMenu(FMenuBuilder& M
 	AssetPickerConfig.bAllowNullSelection = false;
 	AssetPickerConfig.InitialAssetViewType = EAssetViewType::List;
 	AssetPickerConfig.Filter.bRecursiveClasses = true;
-	AssetPickerConfig.Filter.ClassNames.Add(UControlRigSequence::StaticClass()->GetFName());
+	AssetPickerConfig.Filter.ClassPaths.Add(UControlRigSequence::StaticClass()->GetClassPathName());
 	AssetPickerConfig.SaveSettingsName = TEXT("SequencerAssetPicker");
 	}
 
@@ -1310,7 +1310,7 @@ void FControlRigParameterTrackEditor::HandleAddControlRigSubMenu(FMenuBuilder& M
 		AssetPickerConfig.InitialAssetViewType = EAssetViewType::List;
 		AssetPickerConfig.OnShouldFilterAsset = FOnShouldFilterAsset::CreateRaw(this, &FControlRigParameterTrackEditor::ShouldFilterAsset);
 		AssetPickerConfig.Filter.bRecursiveClasses = true;
-		AssetPickerConfig.Filter.ClassNames.Add((UControlRig::StaticClass())->GetFName());
+		AssetPickerConfig.Filter.ClassPaths.Add((UControlRig::StaticClass())->GetClassPathName());
 		AssetPickerConfig.Filter.TagsAndValues.Add(TEXT("Skeleton"), FAssetData(Skeleton).GetExportTextName());
 		AssetPickerConfig.SaveSettingsName = TEXT("SequencerAssetPicker");
 
@@ -1334,7 +1334,7 @@ void FControlRigParameterTrackEditor::HandleAddControlRigSubMenu(FMenuBuilder& M
 /*
 bool FControlRigParameterTrackEditor::ShouldFilterAsset(const FAssetData& AssetData)
 {
-	if (AssetData.AssetClass == UControlRig::StaticClass()->GetFName())
+	if (AssetData.AssetClassPath == UControlRig::StaticClass()->GetClassPathName())
 	{
 		return true;
 	}
@@ -3618,7 +3618,7 @@ void FControlRigParameterSection::BuildSectionContextMenu(FMenuBuilder& MenuBuil
 
 				// Collect a full list of assets with the specified class
 				TArray<FAssetData> AssetDataList;
-				AssetRegistryModule.Get().GetAssetsByClass(UAnimSequenceBase::StaticClass()->GetFName(), AssetDataList, true);
+				AssetRegistryModule.Get().GetAssetsByClass(UAnimSequenceBase::StaticClass()->GetClassPathName(), AssetDataList, true);
 
 				if (AssetDataList.Num())
 				{
@@ -3969,7 +3969,7 @@ void FControlRigParameterSection::AddAnimationSubMenuForFK(FMenuBuilder& MenuBui
 		AssetPickerConfig.InitialAssetViewType = EAssetViewType::List;
 		AssetPickerConfig.OnShouldFilterAsset = FOnShouldFilterAsset::CreateRaw(this, &FControlRigParameterSection::ShouldFilterAssetForFK);
 		AssetPickerConfig.Filter.bRecursiveClasses = true;
-		AssetPickerConfig.Filter.ClassNames.Add(UAnimSequenceBase::StaticClass()->GetFName());
+		AssetPickerConfig.Filter.ClassPaths.Add(UAnimSequenceBase::StaticClass()->GetClassPathName());
 		AssetPickerConfig.Filter.TagsAndValues.Add(TEXT("Skeleton"), FAssetData(Skeleton).GetExportTextName());
 		AssetPickerConfig.SaveSettingsName = TEXT("SequencerAssetPicker");
 		AssetPickerConfig.AdditionalReferencingAssets.Add(FAssetData(Sequence));
@@ -4022,7 +4022,7 @@ bool FControlRigParameterSection::ShouldFilterAssetForFK(const FAssetData& Asset
 {
 	// we don't want 
 
-	if (AssetData.AssetClass == UAnimMontage::StaticClass()->GetFName())
+	if (AssetData.AssetClassPath == UAnimMontage::StaticClass()->GetClassPathName())
 	{
 		return true;
 	}
diff --git a/Engine/Plugins/Animation/GameplayInsights/Source/GameplayInsights/Private/GameplayProvider.cpp b/Engine/Plugins/Animation/GameplayInsights/Source/GameplayInsights/Private/GameplayProvider.cpp
--- a/Engine/Plugins/Animation/GameplayInsights/Source/GameplayInsights/Private/GameplayProvider.cpp
+++ b/Engine/Plugins/Animation/GameplayInsights/Source/GameplayInsights/Private/GameplayProvider.cpp
@@ -29,7 +29,7 @@ static UClass* FindBlueprintClass(const FString& TargetNameRaw)
 
 	FARFilter Filter;
 	Filter.bRecursiveClasses = true;
-	Filter.ClassNames.Add(UBlueprintCore::StaticClass()->GetFName());
+	Filter.ClassPaths.Add(UBlueprintCore::StaticClass()->GetClassPathName());
 
 	// We enumerate all assets to find any blueprints who inherit from native classes directly - or
 	// from other blueprints.
@@ -53,9 +53,9 @@ static UClass* FindBlueprintClass(const FString& TargetNameRaw)
 
 // This is copied from EditorUtilitySubsystem (where it is not public)
 // Should probably be somewhere shared
-static UClass* FindClassByName(const FString& RawTargetName)
+static UClass* FindClassByPathName(const FString& RawTargePathtName)
 {
-	FString TargetName = RawTargetName;
+	FString TargetName = RawTargePathtName;
 
 	// Check native classes and loaded assets first before resorting to the asset registry
 	bool bIsValidClassName = true;
@@ -90,14 +90,7 @@ static UClass* FindClassByName(const FString& RawTargetName)
 	UClass* ResultClass = nullptr;
 	if (bIsValidClassName)
 	{
-		if (FPackageName::IsShortPackageName(TargetName))
-		{
-			ResultClass = FindObject<UClass>(ANY_PACKAGE, *TargetName);
-		}
-		else
-		{
-			ResultClass = FindObject<UClass>(nullptr, *TargetName);
-		}
+		ResultClass = FindObject<UClass>(nullptr, *TargetName);
 	}
 
 	// If we still haven't found anything yet, try the asset registry for blueprints that match the requirements
@@ -228,7 +221,7 @@ const UClass* FGameplayProvider::FindClass(uint64 InClassId) const
 #if WITH_EDITOR
 	if (const FClassInfo* ClassInfo = FindClassInfo(InClassId))
 	{
-		return FindClassByName(ClassInfo->Name);
+		return FindClassByPathName(ClassInfo->PathName);
 	}
 	return nullptr;
 #else
diff --git a/Engine/Plugins/Animation/GameplayInsights/Source/GameplayInsights/Private/SVariantValueView.cpp b/Engine/Plugins/Animation/GameplayInsights/Source/GameplayInsights/Private/SVariantValueView.cpp
--- a/Engine/Plugins/Animation/GameplayInsights/Source/GameplayInsights/Private/SVariantValueView.cpp
+++ b/Engine/Plugins/Animation/GameplayInsights/Source/GameplayInsights/Private/SVariantValueView.cpp
@@ -198,7 +198,7 @@ static TSharedRef<SWidget> MakeVariantValueWidget(const TraceServices::IAnalysis
 					}
 					else // it must be a native class
 					{
-						if (UClass* FoundClass = FindObject<UClass>(ANY_PACKAGE, ClassInfo.Name))
+						if (UClass* FoundClass = UClass::TryFindTypeSlow<UClass>(ClassInfo.Name))
 						{
 							FSourceCodeNavigation::NavigateToClass(FoundClass);
 						}
diff --git a/Engine/Plugins/Animation/IKRig/Source/IKRigEditor/Private/RetargetEditor/IKRetargetEditorController.cpp b/Engine/Plugins/Animation/IKRig/Source/IKRigEditor/Private/RetargetEditor/IKRetargetEditorController.cpp
--- a/Engine/Plugins/Animation/IKRig/Source/IKRigEditor/Private/RetargetEditor/IKRetargetEditorController.cpp
+++ b/Engine/Plugins/Animation/IKRig/Source/IKRigEditor/Private/RetargetEditor/IKRetargetEditorController.cpp
@@ -536,7 +536,7 @@ void FIKRetargetEditorController::HandleImportPose()
 
 	// the asset picker will only show animation sequences
 	FAssetPickerConfig AssetPickerConfig;
-	AssetPickerConfig.Filter.ClassNames.Add(URetargetPose::StaticClass()->GetFName());
+	AssetPickerConfig.Filter.ClassPaths.Add(URetargetPose::StaticClass()->GetClassPathName());
 	AssetPickerConfig.Filter.bRecursiveClasses = true;
 	AssetPickerConfig.OnAssetSelected = FOnAssetSelected::CreateSP(this, &FIKRetargetEditorController::OnRetargetPoseSelected);
 	AssetPickerConfig.InitialAssetViewType = EAssetViewType::Tile;
@@ -644,9 +644,9 @@ void FIKRetargetEditorController::HandleImportPoseFromSequence()
 
 	// the asset picker will only show animation sequences compatible with the preview mesh
 	FAssetPickerConfig AssetPickerConfig;
-	AssetPickerConfig.Filter.ClassNames.Add(UAnimSequence::StaticClass()->GetFName());
-	AssetPickerConfig.Filter.ClassNames.Add(UAnimMontage::StaticClass()->GetFName());
-	AssetPickerConfig.Filter.ClassNames.Add(UPoseAsset::StaticClass()->GetFName());
+	AssetPickerConfig.Filter.ClassPaths.Add(UAnimSequence::StaticClass()->GetClassPathName());
+	AssetPickerConfig.Filter.ClassPaths.Add(UAnimMontage::StaticClass()->GetClassPathName());
+	AssetPickerConfig.Filter.ClassPaths.Add(UPoseAsset::StaticClass()->GetClassPathName());
 	AssetPickerConfig.InitialAssetViewType = EAssetViewType::Column;
 	AssetPickerConfig.bAddFilterUI = true;
 	AssetPickerConfig.bShowPathInColumnView = true;
@@ -928,7 +928,7 @@ void FIKRetargetEditorController::HandleExportPose()
 	FSaveAssetDialogConfig SaveAssetDialogConfig;
 	SaveAssetDialogConfig.DefaultPath = AssetController->GetAsset()->GetPackage()->GetPathName();
 	SaveAssetDialogConfig.DefaultAssetName = GetCurrentPoseName().ToString();
-	SaveAssetDialogConfig.AssetClassNames.Add(URetargetPose::StaticClass()->GetFName());
+	SaveAssetDialogConfig.AssetClassNames.Add(URetargetPose::StaticClass()->GetClassPathName());
 	SaveAssetDialogConfig.ExistingAssetPolicy = ESaveAssetDialogExistingAssetPolicy::AllowButWarn;
 	SaveAssetDialogConfig.DialogTitleOverride = LOCTEXT("ExportRetargetPoseDialogTitle", "Export Retarget Pose");
 
diff --git a/Engine/Plugins/Animation/IKRig/Source/IKRigEditor/Private/RetargetEditor/IKRetargetFactory.cpp b/Engine/Plugins/Animation/IKRig/Source/IKRigEditor/Private/RetargetEditor/IKRetargetFactory.cpp
--- a/Engine/Plugins/Animation/IKRig/Source/IKRigEditor/Private/RetargetEditor/IKRetargetFactory.cpp
+++ b/Engine/Plugins/Animation/IKRig/Source/IKRigEditor/Private/RetargetEditor/IKRetargetFactory.cpp
@@ -71,7 +71,7 @@ bool UIKRetargetFactory::ConfigureProperties()
 	FAssetPickerConfig AssetPickerConfig;
 
 	/** The asset picker will only show skeletal meshes */
-	AssetPickerConfig.Filter.ClassNames.Add(UIKRigDefinition::StaticClass()->GetFName());
+	AssetPickerConfig.Filter.ClassPaths.Add(UIKRigDefinition::StaticClass()->GetClassPathName());
 	AssetPickerConfig.Filter.bRecursiveClasses = true;
 
 	/** The delegate that fires when an asset was selected */
diff --git a/Engine/Plugins/Animation/IKRig/Source/IKRigEditor/Private/RetargetEditor/SIKRetargetAssetBrowser.cpp b/Engine/Plugins/Animation/IKRig/Source/IKRigEditor/Private/RetargetEditor/SIKRetargetAssetBrowser.cpp
--- a/Engine/Plugins/Animation/IKRig/Source/IKRigEditor/Private/RetargetEditor/SIKRetargetAssetBrowser.cpp
+++ b/Engine/Plugins/Animation/IKRig/Source/IKRigEditor/Private/RetargetEditor/SIKRetargetAssetBrowser.cpp
@@ -56,9 +56,9 @@ void SIKRetargetAssetBrowser::RefreshView()
 	FAssetPickerConfig AssetPickerConfig;
 
 	// setup filtering
-	AssetPickerConfig.Filter.ClassNames.Add(UAnimSequence::StaticClass()->GetFName());
-	AssetPickerConfig.Filter.ClassNames.Add(UAnimMontage::StaticClass()->GetFName());
-	AssetPickerConfig.Filter.ClassNames.Add(UPoseAsset::StaticClass()->GetFName());
+	AssetPickerConfig.Filter.ClassPaths.Add(UAnimSequence::StaticClass()->GetClassPathName());
+	AssetPickerConfig.Filter.ClassPaths.Add(UAnimMontage::StaticClass()->GetClassPathName());
+	AssetPickerConfig.Filter.ClassPaths.Add(UPoseAsset::StaticClass()->GetClassPathName());
 	AssetPickerConfig.InitialAssetViewType = EAssetViewType::Column;
 	AssetPickerConfig.bAddFilterUI = true;
 	AssetPickerConfig.bShowPathInColumnView = true;
diff --git a/Engine/Plugins/Animation/IKRig/Source/IKRigEditor/Private/RigEditor/IKRigDefinitionFactory.cpp b/Engine/Plugins/Animation/IKRig/Source/IKRigEditor/Private/RigEditor/IKRigDefinitionFactory.cpp
--- a/Engine/Plugins/Animation/IKRig/Source/IKRigEditor/Private/RigEditor/IKRigDefinitionFactory.cpp
+++ b/Engine/Plugins/Animation/IKRig/Source/IKRigEditor/Private/RigEditor/IKRigDefinitionFactory.cpp
@@ -66,7 +66,7 @@ bool UIKRigDefinitionFactory::ConfigureProperties()
 	FAssetPickerConfig AssetPickerConfig;
 
 	/** The asset picker will only show skeletal meshes */
-	AssetPickerConfig.Filter.ClassNames.Add(USkeletalMesh::StaticClass()->GetFName());
+	AssetPickerConfig.Filter.ClassPaths.Add(USkeletalMesh::StaticClass()->GetClassPathName());
 
 	/** The delegate that fires when an asset was selected */
 	AssetPickerConfig.OnAssetSelected = FOnAssetSelected::CreateUObject(this, &UIKRigDefinitionFactory::OnSkeletalMeshSelected);
diff --git a/Engine/Plugins/Animation/IKRig/Source/IKRigEditor/Private/RigEditor/SIKRigAssetBrowser.cpp b/Engine/Plugins/Animation/IKRig/Source/IKRigEditor/Private/RigEditor/SIKRigAssetBrowser.cpp
--- a/Engine/Plugins/Animation/IKRig/Source/IKRigEditor/Private/RigEditor/SIKRigAssetBrowser.cpp
+++ b/Engine/Plugins/Animation/IKRig/Source/IKRigEditor/Private/RigEditor/SIKRigAssetBrowser.cpp
@@ -37,9 +37,9 @@ void SIKRigAssetBrowser::RefreshView()
 	FAssetPickerConfig AssetPickerConfig;
 	
 	// setup filtering
-	AssetPickerConfig.Filter.ClassNames.Add(UAnimSequence::StaticClass()->GetFName());
-	AssetPickerConfig.Filter.ClassNames.Add(UAnimMontage::StaticClass()->GetFName());
-	AssetPickerConfig.Filter.ClassNames.Add(UPoseAsset::StaticClass()->GetFName());
+	AssetPickerConfig.Filter.ClassPaths.Add(UAnimSequence::StaticClass()->GetClassPathName());
+	AssetPickerConfig.Filter.ClassPaths.Add(UAnimMontage::StaticClass()->GetClassPathName());
+	AssetPickerConfig.Filter.ClassPaths.Add(UPoseAsset::StaticClass()->GetClassPathName());
 	AssetPickerConfig.InitialAssetViewType = EAssetViewType::Column;
 	AssetPickerConfig.bAddFilterUI = true;
 	AssetPickerConfig.bShowPathInColumnView = true;
diff --git a/Engine/Plugins/Animation/LiveLink/Source/LiveLinkEditor/Private/LiveLinkClientPanelToolbar.cpp b/Engine/Plugins/Animation/LiveLink/Source/LiveLinkEditor/Private/LiveLinkClientPanelToolbar.cpp
--- a/Engine/Plugins/Animation/LiveLink/Source/LiveLinkEditor/Private/LiveLinkClientPanelToolbar.cpp
+++ b/Engine/Plugins/Animation/LiveLink/Source/LiveLinkEditor/Private/LiveLinkClientPanelToolbar.cpp
@@ -753,7 +753,7 @@ TSharedRef<SWidget> SLiveLinkClientPanelToolbar::OnPresetGeneratePresetsMenu()
 		AssetPickerConfig.bSortByPathInColumnView = false;
 
 		AssetPickerConfig.AssetShowWarningText = LOCTEXT("NoPresets_Warning", "No Presets Found");
-		AssetPickerConfig.Filter.ClassNames.Add(ULiveLinkPreset::StaticClass()->GetFName());
+		AssetPickerConfig.Filter.ClassPaths.Add(ULiveLinkPreset::StaticClass()->GetClassPathName());
 		AssetPickerConfig.OnAssetSelected = FOnAssetSelected::CreateSP(this, &SLiveLinkClientPanelToolbar::OnImportPreset);
 	}
 
@@ -779,7 +779,7 @@ static bool OpenSaveDialog(const FString& InDefaultPath, const FString& InNewNam
 	{
 		SaveAssetDialogConfig.DefaultPath = InDefaultPath;
 		SaveAssetDialogConfig.DefaultAssetName = InNewNameSuggestion;
-		SaveAssetDialogConfig.AssetClassNames.Add(ULiveLinkPreset::StaticClass()->GetFName());
+		SaveAssetDialogConfig.AssetClassNames.Add(ULiveLinkPreset::StaticClass()->GetClassPathName());
 		SaveAssetDialogConfig.ExistingAssetPolicy = ESaveAssetDialogExistingAssetPolicy::AllowButWarn;
 		SaveAssetDialogConfig.DialogTitleOverride = LOCTEXT("SaveLiveLinkPresetDialogTitle", "Save LiveLink Preset");
 	}
diff --git a/Engine/Plugins/Animation/RigLogic/Source/RigLogicEditor/Private/DNAImporter.cpp b/Engine/Plugins/Animation/RigLogic/Source/RigLogicEditor/Private/DNAImporter.cpp
--- a/Engine/Plugins/Animation/RigLogic/Source/RigLogicEditor/Private/DNAImporter.cpp
+++ b/Engine/Plugins/Animation/RigLogic/Source/RigLogicEditor/Private/DNAImporter.cpp
@@ -36,7 +36,7 @@ FDNAAssetImportOptions* GetImportOptions(FDNAImporter * DNAImporter, UDNAAssetIm
 			// Look in the current target directory to see if we have a skeleton
 			FARFilter Filter;
 			Filter.PackagePaths.Add(*FPaths::GetPath(FullPath));
-			Filter.ClassNames.Add(USkeletalMesh::StaticClass()->GetFName());
+			Filter.ClassPaths.Add(USkeletalMesh::StaticClass()->GetClassPathName());
 
 			IAssetRegistry& AssetRegistry = FModuleManager::LoadModuleChecked<FAssetRegistryModule>("AssetRegistry").Get();
 			TArray<FAssetData> SkeletalMeshAssets;
diff --git a/Engine/Plugins/Bridge/Source/MegascansPlugin/Private/Tools/BlendMaterials.cpp b/Engine/Plugins/Bridge/Source/MegascansPlugin/Private/Tools/BlendMaterials.cpp
--- a/Engine/Plugins/Bridge/Source/MegascansPlugin/Private/Tools/BlendMaterials.cpp
+++ b/Engine/Plugins/Bridge/Source/MegascansPlugin/Private/Tools/BlendMaterials.cpp
@@ -58,7 +58,7 @@ void FMaterialBlend::BlendSelectedMaterials()
 
 	const UMaterialBlendSettings* BlendSettings = GetDefault<UMaterialBlendSettings>();
 
-	TArray<UMaterialInstanceConstant*> SelectedMaterialInstances = AssetUtils::GetSelectedAssets(TEXT("MaterialInstanceConstant"));
+	TArray<UMaterialInstanceConstant*> SelectedMaterialInstances = AssetUtils::GetSelectedAssets(FTopLevelAssetPath(TEXT("/Script/Engine"), TEXT("MaterialInstanceConstant")));
 
 	if (SelectedMaterialInstances.Num() < 2)
 	{
diff --git a/Engine/Plugins/Bridge/Source/MegascansPlugin/Private/Utilities/MiscUtils.cpp b/Engine/Plugins/Bridge/Source/MegascansPlugin/Private/Utilities/MiscUtils.cpp
--- a/Engine/Plugins/Bridge/Source/MegascansPlugin/Private/Utilities/MiscUtils.cpp
+++ b/Engine/Plugins/Bridge/Source/MegascansPlugin/Private/Utilities/MiscUtils.cpp
@@ -266,8 +266,10 @@ void CopyMSPresets()
 		return AssetMetaData;
 	}
 
-	TArray<UMaterialInstanceConstant*> AssetUtils::GetSelectedAssets(const FString& AssetClass)
+	TArray<UMaterialInstanceConstant*> AssetUtils::GetSelectedAssets(const FTopLevelAssetPath& AssetClass)
 	{
+		check(!AssetClass.IsNull());
+
 		TArray<UMaterialInstanceConstant*> ObjectArray;
 
 		TArray<FAssetData> AssetDatas;
@@ -278,7 +280,7 @@ void CopyMSPresets()
 		for (FAssetData SelectedAsset : AssetDatas)
 		{
 
-			if (SelectedAsset.AssetClass == FName(*AssetClass))
+			if (SelectedAsset.AssetClassPath == AssetClass)
 			{
 				ObjectArray.Add(CastChecked<UMaterialInstanceConstant>(UEditorAssetLibrary::LoadAsset(SelectedAsset.ObjectPath.ToString())));
 			}
diff --git a/Engine/Plugins/Bridge/Source/MegascansPlugin/Private/Utilities/MiscUtils.h b/Engine/Plugins/Bridge/Source/MegascansPlugin/Private/Utilities/MiscUtils.h
--- a/Engine/Plugins/Bridge/Source/MegascansPlugin/Private/Utilities/MiscUtils.h
+++ b/Engine/Plugins/Bridge/Source/MegascansPlugin/Private/Utilities/MiscUtils.h
@@ -33,7 +33,7 @@ namespace AssetUtils {	//template<typename T>
 	void DeleteDirectory(FString TargetDirectory);
 	bool DeleteAsset(const FString& AssetPath);
 	FUAssetMeta GetAssetMetaData(const FString& JsonPath);
-	TArray<UMaterialInstanceConstant*> GetSelectedAssets(const FString& AssetClass);
+	TArray<UMaterialInstanceConstant*> GetSelectedAssets(const FTopLevelAssetPath& AssetClass);
 	void AddFoliageTypesToLevel(TArray<FString> FoliageTypePaths);
 
 	//To manage import settings like Auto-Populate Foliage types etc.
diff --git a/Engine/Plugins/Compositing/Composure/Source/ComposureLayersEditor/Private/CompElementManager.cpp b/Engine/Plugins/Compositing/Composure/Source/ComposureLayersEditor/Private/CompElementManager.cpp
--- a/Engine/Plugins/Compositing/Composure/Source/ComposureLayersEditor/Private/CompElementManager.cpp
+++ b/Engine/Plugins/Compositing/Composure/Source/ComposureLayersEditor/Private/CompElementManager.cpp
@@ -474,7 +474,7 @@ UCompositingMediaCaptureOutput* FCompElementManager::ResetMediaCapture(const FNa
 		{
 			if (ClassIt->IsChildOf(UMediaOutput::StaticClass()) && !ClassIt->HasAnyClassFlags(CLASS_Abstract))
 			{
-				SelectAssetConfig.AssetClassNames.Add(ClassIt->GetFName());
+				SelectAssetConfig.AssetClassNames.Add(ClassIt->GetClassPathName());
 			}
 		}
 
diff --git a/Engine/Plugins/Developer/Concert/ConcertMain/Source/ConcertTransport/Private/ConcertLocalEndpoint.cpp b/Engine/Plugins/Developer/Concert/ConcertMain/Source/ConcertTransport/Private/ConcertLocalEndpoint.cpp
--- a/Engine/Plugins/Developer/Concert/ConcertMain/Source/ConcertTransport/Private/ConcertLocalEndpoint.cpp
+++ b/Engine/Plugins/Developer/Concert/ConcertMain/Source/ConcertTransport/Private/ConcertLocalEndpoint.cpp
@@ -150,27 +150,27 @@ FOnConcertRemoteEndpointConnectionChanged& FConcertLocalEndpoint::OnRemoteEndpoi
 	return OnRemoteEndpointConnectionChangedDelegate;
 }
 
-void FConcertLocalEndpoint::InternalAddRequestHandler(const FName& RequestMessageType, const TSharedRef<IConcertRequestHandler>& Handler)
+void FConcertLocalEndpoint::InternalAddRequestHandler(const FTopLevelAssetPath& RequestMessageType, const TSharedRef<IConcertRequestHandler>& Handler)
 {
 	RequestHandlers.Add(RequestMessageType, Handler);
 }
 
-void FConcertLocalEndpoint::InternalAddEventHandler(const FName& EventMessageType, const TSharedRef<IConcertEventHandler>& Handler)
+void FConcertLocalEndpoint::InternalAddEventHandler(const FTopLevelAssetPath& EventMessageType, const TSharedRef<IConcertEventHandler>& Handler)
 {
 	EventHandlers.Add(EventMessageType, Handler);
 }
 
-void FConcertLocalEndpoint::InternalRemoveRequestHandler(const FName& RequestMessageType)
+void FConcertLocalEndpoint::InternalRemoveRequestHandler(const FTopLevelAssetPath& RequestMessageType)
 {
 	RequestHandlers.Remove(RequestMessageType);
 }
 
-void FConcertLocalEndpoint::InternalRemoveEventHandler(const FName& EventMessageType)
+void FConcertLocalEndpoint::InternalRemoveEventHandler(const FTopLevelAssetPath& EventMessageType)
 {
 	EventHandlers.Remove(EventMessageType);
 }
 
-void FConcertLocalEndpoint::InternalSubscribeToEvent(const FName& EventMessageType)
+void FConcertLocalEndpoint::InternalSubscribeToEvent(const FTopLevelAssetPath& EventMessageType)
 {
 	if (MessageEndpoint.IsValid())
 	{
@@ -178,7 +178,7 @@ void FConcertLocalEndpoint::InternalSubscribeToEvent(const FName& EventMessageTy
 	}
 }
 
-void FConcertLocalEndpoint::InternalUnsubscribeFromEvent(const FName& EventMessageType)
+void FConcertLocalEndpoint::InternalUnsubscribeFromEvent(const FTopLevelAssetPath& EventMessageType)
 {
 	if (MessageEndpoint.IsValid())
 	{
@@ -669,7 +669,7 @@ void FConcertLocalEndpoint::HandleMessage(const FConcertMessageContext& ConcertC
 void FConcertLocalEndpoint::ProcessEvent(const FConcertMessageContext& ConcertContext)
 {
 	const FConcertEventData* Event = ConcertContext.GetMessage<FConcertEventData>();
-	const FName EventType = ConcertContext.MessageType->GetFName();
+	const FTopLevelAssetPath EventType = ConcertContext.MessageType->GetStructPathName();
 
 	Logger.LogProcessEvent(ConcertContext, EndpointContext.EndpointId);
 	TSharedPtr<IConcertEventHandler> Handler = EventHandlers.FindRef(EventType);
@@ -687,7 +687,7 @@ void FConcertLocalEndpoint::ProcessEvent(const FConcertMessageContext& ConcertCo
 void FConcertLocalEndpoint::ProcessRequest(const FConcertMessageContext& ConcertContext)
 {
 	const FConcertRequestData* Request = ConcertContext.GetMessage<FConcertRequestData>();
-	const FName RequestType = ConcertContext.MessageType->GetFName();
+	const FTopLevelAssetPath RequestType = ConcertContext.MessageType->GetStructPathName();
 
 	// The response ID should match the request message, and the response should go back to the endpoint where the request came from
 	auto DispatchResponse = [this, RequestMessageId = ConcertContext.Message->MessageId, ResponseDestinationEndpointId = Request->ConcertEndpointId](TSharedPtr<IConcertResponse> Response)
diff --git a/Engine/Plugins/Developer/Concert/ConcertMain/Source/ConcertTransport/Private/ConcertLocalEndpoint.h b/Engine/Plugins/Developer/Concert/ConcertMain/Source/ConcertTransport/Private/ConcertLocalEndpoint.h
--- a/Engine/Plugins/Developer/Concert/ConcertMain/Source/ConcertTransport/Private/ConcertLocalEndpoint.h
+++ b/Engine/Plugins/Developer/Concert/ConcertMain/Source/ConcertTransport/Private/ConcertLocalEndpoint.h
@@ -36,14 +36,14 @@ public:
 	virtual FOnConcertRemoteEndpointConnectionChanged& OnRemoteEndpointConnectionChanged() override;
 
 protected:
-	virtual void InternalAddRequestHandler(const FName& RequestMessageType, const TSharedRef<IConcertRequestHandler>& Handler) override;
-	virtual void InternalAddEventHandler(const FName& EventMessageType, const TSharedRef<IConcertEventHandler>& Handler) override;
+	virtual void InternalAddRequestHandler(const FTopLevelAssetPath& RequestMessageType, const TSharedRef<IConcertRequestHandler>& Handler) override;
+	virtual void InternalAddEventHandler(const FTopLevelAssetPath& EventMessageType, const TSharedRef<IConcertEventHandler>& Handler) override;
 
-	virtual void InternalRemoveRequestHandler(const FName& RequestMessageType) override;
-	virtual void InternalRemoveEventHandler(const FName& EventMessageType) override;
+	virtual void InternalRemoveRequestHandler(const FTopLevelAssetPath& RequestMessageType) override;
+	virtual void InternalRemoveEventHandler(const FTopLevelAssetPath& EventMessageType) override;
 
-	virtual void InternalSubscribeToEvent(const FName& EventMessageType) override;
-	virtual void InternalUnsubscribeFromEvent(const FName& EventMessageType) override;
+	virtual void InternalSubscribeToEvent(const FTopLevelAssetPath& EventMessageType) override;
+	virtual void InternalUnsubscribeFromEvent(const FTopLevelAssetPath& EventMessageType) override;
 
 	virtual void InternalQueueRequest(const TSharedRef<IConcertRequest>& Request, const FGuid& Endpoint) override;
 	virtual void InternalQueueResponse(const TSharedRef<IConcertResponse>& Response, const FGuid& Endpoint) override;
@@ -159,10 +159,10 @@ private:
 	FOnConcertRemoteEndpointConnectionChanged OnRemoteEndpointConnectionChangedDelegate;
 
 	/** Registered message handlers that do not generate a response */
-	TMap<FName, TSharedPtr<IConcertEventHandler>> EventHandlers;
+	TMap<FTopLevelAssetPath, TSharedPtr<IConcertEventHandler>> EventHandlers;
 
 	/** Registered message handlers that returns a response */
-	TMap<FName, TSharedPtr<IConcertRequestHandler>> RequestHandlers;
+	TMap<FTopLevelAssetPath, TSharedPtr<IConcertRequestHandler>> RequestHandlers;
 
 	/** Handle to the registered ticker for request and reliable message */
 	FTSTicker::FDelegateHandle TickerHandle;
diff --git a/Engine/Plugins/Developer/Concert/ConcertMain/Source/ConcertTransport/Public/IConcertEndpoint.h b/Engine/Plugins/Developer/Concert/ConcertMain/Source/ConcertTransport/Public/IConcertEndpoint.h
--- a/Engine/Plugins/Developer/Concert/ConcertMain/Source/ConcertTransport/Public/IConcertEndpoint.h
+++ b/Engine/Plugins/Developer/Concert/ConcertMain/Source/ConcertTransport/Public/IConcertEndpoint.h
@@ -149,7 +149,7 @@ public:
 	void RegisterRequestHandler(typename TConcertFunctionRequestHandler<ResponseType>::FFuncType Func)
 	{
 		static_assert(TIsDerivedFrom<RequestType, FConcertRequestData>::IsDerived, "RequestType need to be a UStruct deriving of FConcertRequestData.");
-		InternalAddRequestHandler(RequestType::StaticStruct()->GetFName(), MakeShared<TConcertFunctionRequestHandler<ResponseType>>(MoveTemp(Func)));
+		InternalAddRequestHandler(RequestType::StaticStruct()->GetStructPathName(), MakeShared<TConcertFunctionRequestHandler<ResponseType>>(MoveTemp(Func)));
 	}
 
 	/**
@@ -161,7 +161,7 @@ public:
 	void RegisterRequestHandler(HandlerType* Handler, typename TConcertRawRequestHandler<ResponseType, HandlerType>::FFuncType Func)
 	{
 		static_assert(TIsDerivedFrom<RequestType, FConcertRequestData>::IsDerived, "RequestType need to be a UStruct deriving of FConcertRequestData.");
-		InternalAddRequestHandler(RequestType::StaticStruct()->GetFName(), MakeShared<TConcertRawRequestHandler<ResponseType, HandlerType>>(Handler, Func));
+		InternalAddRequestHandler(RequestType::StaticStruct()->GetStructPathName(), MakeShared<TConcertRawRequestHandler<ResponseType, HandlerType>>(Handler, Func));
 	}
 
 	/**
@@ -171,7 +171,7 @@ public:
 	void UnregisterRequestHandler()
 	{
 		static_assert(TIsDerivedFrom<RequestType, FConcertRequestData>::IsDerived, "RequestType need to be a UStruct deriving of FConcertRequestData.");
-		InternalRemoveRequestHandler(RequestType::StaticStruct()->GetFName());
+		InternalRemoveRequestHandler(RequestType::StaticStruct()->GetStructPathName());
 	}
 
 	/**
@@ -182,7 +182,7 @@ public:
 	void RegisterEventHandler(typename TConcertFunctionEventHandler::FFuncType Func)
 	{
 		static_assert(TIsDerivedFrom<EventType, FConcertEventData>::IsDerived, "EventType need to be a UStruct deriving of FConcertEventData.");
-		InternalAddEventHandler(EventType::StaticStruct()->GetFName(), MakeShared<TConcertFunctionEventHandler>(MoveTemp(Func)));
+		InternalAddEventHandler(EventType::StaticStruct()->GetStructPathName(), MakeShared<TConcertFunctionEventHandler>(MoveTemp(Func)));
 	}
 
 	/**
@@ -194,7 +194,7 @@ public:
 	void RegisterEventHandler(HandlerType* Handler, typename TConcertRawEventHandler<HandlerType>::FFuncType Func)
 	{
 		static_assert(TIsDerivedFrom<EventType, FConcertEventData>::IsDerived, "EventType need to be a UStruct deriving of FConcertEventData.");
-		InternalAddEventHandler(EventType::StaticStruct()->GetFName(), MakeShared<TConcertRawEventHandler<HandlerType>>(Handler, Func));
+		InternalAddEventHandler(EventType::StaticStruct()->GetStructPathName(), MakeShared<TConcertRawEventHandler<HandlerType>>(Handler, Func));
 	}
 
 	/**
@@ -204,7 +204,7 @@ public:
 	void UnregisterEventHandler()
 	{
 		static_assert(TIsDerivedFrom<EventType, FConcertEventData>::IsDerived, "EventType need to be a UStruct deriving of FConcertEventData.");
-		InternalRemoveEventHandler(EventType::StaticStruct()->GetFName());
+		InternalRemoveEventHandler(EventType::StaticStruct()->GetStructPathName());
 	}
 
 	/**
@@ -216,8 +216,8 @@ public:
 	void SubscribeEventHandler(HandlerType* Handler, typename TConcertRawEventHandler<HandlerType>::FFuncType Func)
 	{
 		static_assert(TIsDerivedFrom<EventType, FConcertEventData>::IsDerived, "EventType need to be a UStruct deriving of FConcertEventData.");
-		InternalAddEventHandler(EventType::StaticStruct()->GetFName(), MakeShared<TConcertRawEventHandler<HandlerType>>(Handler, Func));
-		InternalSubscribeToEvent(EventType::StaticStruct()->GetFName());
+		InternalAddEventHandler(EventType::StaticStruct()->GetStructPathName(), MakeShared<TConcertRawEventHandler<HandlerType>>(Handler, Func));
+		InternalSubscribeToEvent(EventType::StaticStruct()->GetStructPathName());
 	}
 
 	// TODO TFunction EventHandler
@@ -231,7 +231,7 @@ public:
 		static_assert(TIsDerivedFrom<EventType, FConcertEventData>::IsDerived, "EventType need to be a UStruct deriving of FConcertEventData.");
 
 		// Names can be invalid if unregistering during shutdown
-		const FName EventName = EventType::StaticStruct()->GetFName();
+		const FTopLevelAssetPath EventName = EventType::StaticStruct()->GetStructPathName();
 		if (EventName.IsValid())
 		{
 			InternalRemoveEventHandler(EventName);
@@ -257,32 +257,86 @@ protected:
 	/** 
 	 * Add a Request Handler
 	 */
-	virtual void InternalAddRequestHandler(const FName& RequestMessageType, const TSharedRef<IConcertRequestHandler>& Handler) = 0;
+	UE_DEPRECATED(5.1, "Types names are now represented by path names. Please use a version of this function that takes an FTopLevelAssetPath as MessageType.")
+	virtual void InternalAddRequestHandler(const FName& RequestMessageType, const TSharedRef<IConcertRequestHandler>& Handler)
+	{
+		InternalAddRequestHandler(UClass::TryConvertShortTypeNameToPathName<UStruct>(RequestMessageType.ToString()), Handler);
+	}
+
+	/**
+	 * Add a Request Handler
+	 */
+	virtual void InternalAddRequestHandler(const FTopLevelAssetPath& RequestMessageType, const TSharedRef<IConcertRequestHandler>& Handler) = 0;
+
+	/**
+	 * Remove an Request Handler
+	 */
+	UE_DEPRECATED(5.1, "Types names are now represented by path names. Please use a version of this function that takes an FTopLevelAssetPath as MessageType.")
+	virtual void InternalRemoveRequestHandler(const FName& RequestMessageType)
+	{
+		InternalRemoveRequestHandler(UClass::TryConvertShortTypeNameToPathName<UStruct>(RequestMessageType.ToString()));
+	}
 
 	/**
 	 * Remove an Request Handler
 	 */
-	virtual void InternalRemoveRequestHandler(const FName& RequestMessageType) = 0;
+	virtual void InternalRemoveRequestHandler(const FTopLevelAssetPath& RequestMessageType) = 0;
+
+	/**
+	 * Add an Event Handler
+	 */
+	UE_DEPRECATED(5.1, "Types names are now represented by path names. Please use a version of this function that takes an FTopLevelAssetPath as MessageType.")
+	virtual void InternalAddEventHandler(const FName& EventMessageType, const TSharedRef<IConcertEventHandler>& Handler)
+	{
+		InternalAddEventHandler(UClass::TryConvertShortTypeNameToPathName<UStruct>(EventMessageType.ToString()), Handler);
+	}
 
 	/**
 	 * Add an Event Handler
 	 */
-	virtual void InternalAddEventHandler(const FName& EventMessageType, const TSharedRef<IConcertEventHandler>& Handler) = 0;
+	virtual void InternalAddEventHandler(const FTopLevelAssetPath& EventMessageType, const TSharedRef<IConcertEventHandler>& Handler) = 0;
 
 	/**
 	 * Remove an Event Handler
 	 */
-	virtual void InternalRemoveEventHandler(const FName& EventMessageType) = 0;
+	UE_DEPRECATED(5.1, "Types names are now represented by path names. Please use a version of this function that takes an FTopLevelAssetPath as MessageType.")
+	virtual void InternalRemoveEventHandler(const FName& EventMessageType)
+	{
+		InternalRemoveEventHandler(UClass::TryConvertShortTypeNameToPathName<UStruct>(EventMessageType.ToString()));
+	}
+
+	/**
+	 * Remove an Event Handler
+	 */
+	virtual void InternalRemoveEventHandler(const FTopLevelAssetPath& EventMessageType) = 0;
+
+	/**
+	 * Subscribe to an Event
+	 */
+	UE_DEPRECATED(5.1, "Types names are now represented by path names. Please use a version of this function that takes an FTopLevelAssetPath as MessageType.")
+	virtual void InternalSubscribeToEvent(const FName& EventMessageType)
+	{
+		InternalSubscribeToEvent(UClass::TryConvertShortTypeNameToPathName<UStruct>(EventMessageType.ToString()));
+	}
 
 	/**
 	 * Subscribe to an Event
 	 */
-	virtual void InternalSubscribeToEvent(const FName& EventMessageType) = 0;
+	virtual void InternalSubscribeToEvent(const FTopLevelAssetPath& EventMessageType) = 0;
+
+	/**
+	 * Unsubscribe from an Event
+	 */
+	UE_DEPRECATED(5.1, "Types names are now represented by path names. Please use a version of this function that takes an FTopLevelAssetPath as MessageType.")
+	virtual void InternalUnsubscribeFromEvent(const FName& EventMessageType)
+	{
+		InternalUnsubscribeFromEvent(UClass::TryConvertShortTypeNameToPathName<UStruct>(EventMessageType.ToString()));
+	}
 
 	/**
 	 * Unsubscribe from an Event
 	 */
-	virtual void InternalUnsubscribeFromEvent(const FName& EventMessageType) = 0;
+	virtual void InternalUnsubscribeFromEvent(const FTopLevelAssetPath& EventMessageType) = 0;
 
 	/**
 	 * Queue a request to be sent to a remote endpoint
diff --git a/Engine/Plugins/Developer/TraceSourceFiltering/Source/SourceFilteringEditor/Private/EditorSessionSourceFilterService.cpp b/Engine/Plugins/Developer/TraceSourceFiltering/Source/SourceFilteringEditor/Private/EditorSessionSourceFilterService.cpp
--- a/Engine/Plugins/Developer/TraceSourceFiltering/Source/SourceFilteringEditor/Private/EditorSessionSourceFilterService.cpp
+++ b/Engine/Plugins/Developer/TraceSourceFiltering/Source/SourceFilteringEditor/Private/EditorSessionSourceFilterService.cpp
@@ -278,7 +278,7 @@ static bool GetSavePresetPackageName(FString& OutPackageName)
 	{
 		SaveAssetDialogConfig.DefaultPath = DefaultPath;
 		SaveAssetDialogConfig.DefaultAssetName = FPaths::GetCleanFilename(UniqueAssetName);
-		SaveAssetDialogConfig.AssetClassNames.Add(USourceFilterCollection::StaticClass()->GetFName());
+		SaveAssetDialogConfig.AssetClassNames.Add(USourceFilterCollection::StaticClass()->GetClassPathName());
 		SaveAssetDialogConfig.ExistingAssetPolicy = ESaveAssetDialogExistingAssetPolicy::AllowButWarn;
 		SaveAssetDialogConfig.DialogTitleOverride = LOCTEXT("SaveSourceFilterPresetDialogTitle", "Save Filter Preset");
 	}
@@ -376,7 +376,7 @@ TSharedPtr<FExtender> FEditorSessionSourceFilterService::GetExtender()
 								AssetPickerConfig.bSortByPathInColumnView = false;
 
 								AssetPickerConfig.AssetShowWarningText = LOCTEXT("NoPresets_Warning", "No Presets Found");
-								AssetPickerConfig.Filter.ClassNames.Add(USourceFilterCollection::StaticClass()->GetFName());
+								AssetPickerConfig.Filter.ClassPaths.Add(USourceFilterCollection::StaticClass()->GetClassPathName());
 								AssetPickerConfig.OnAssetDoubleClicked = FOnAssetDoubleClicked::CreateLambda([this](const FAssetData& InAssetData)
 								{
 									if (UObject* Asset = InAssetData.GetAsset())
diff --git a/Engine/Plugins/Developer/TraceSourceFiltering/Source/SourceFilteringTrace/Private/SourceFilterPresets.cpp b/Engine/Plugins/Developer/TraceSourceFiltering/Source/SourceFilteringTrace/Private/SourceFilterPresets.cpp
--- a/Engine/Plugins/Developer/TraceSourceFiltering/Source/SourceFilteringTrace/Private/SourceFilterPresets.cpp
+++ b/Engine/Plugins/Developer/TraceSourceFiltering/Source/SourceFilteringTrace/Private/SourceFilterPresets.cpp
@@ -47,8 +47,8 @@ void FSourceFilterPresets::ListAvailablePresets()
 void FSourceFilterPresets::GetPresets(TArray<FAssetData>& InOutPresetAssetData)
 {
 	FAssetRegistryModule& AssetRegistryModule = FModuleManager::LoadModuleChecked<FAssetRegistryModule>(TEXT("AssetRegistry"));
-	TMultiMap<FName, FString> TagValues = { { FBlueprintTags::NativeParentClassPath, FString::Printf(TEXT("%s'%s'"), *UClass::StaticClass()->GetName(), *USourceFilterCollection::StaticClass()->GetPathName()) } };
-	AssetRegistryModule.Get().GetAssetsByClass(USourceFilterCollection::StaticClass()->GetFName(), InOutPresetAssetData);
+	TMultiMap<FName, FString> TagValues = { { FBlueprintTags::NativeParentClassPath, FObjectPropertyBase::GetExportPath(USourceFilterCollection::StaticClass()) } };
+	AssetRegistryModule.Get().GetAssetsByClass(USourceFilterCollection::StaticClass()->GetClassPathName(), InOutPresetAssetData);
 }
 
 void FSourceFilterPresets::LoadPreset(const FSoftObjectPath& PresetPath)
diff --git a/Engine/Plugins/Developer/TraceSourceFiltering/Source/SourceFilteringTrace/Private/SourceFilteringTraceModule.cpp b/Engine/Plugins/Developer/TraceSourceFiltering/Source/SourceFilteringTrace/Private/SourceFilteringTraceModule.cpp
--- a/Engine/Plugins/Developer/TraceSourceFiltering/Source/SourceFilteringTrace/Private/SourceFilteringTraceModule.cpp
+++ b/Engine/Plugins/Developer/TraceSourceFiltering/Source/SourceFilteringTrace/Private/SourceFilteringTraceModule.cpp
@@ -33,7 +33,7 @@ void FSourceFilteringTraceModule::TraceFilterClasses()
 	FAssetRegistryModule& AssetRegistryModule = FModuleManager::LoadModuleChecked<FAssetRegistryModule>(TEXT("AssetRegistry"));
 
 	TArray<FAssetData> Filters;
-	TMultiMap<FName, FString> TagValues = { { FBlueprintTags::NativeParentClassPath, FString::Printf(TEXT("%s'%s'"), *UClass::StaticClass()->GetName(), *UDataSourceFilter::StaticClass()->GetPathName()) } };
+	TMultiMap<FName, FString> TagValues = { { FBlueprintTags::NativeParentClassPath, FObjectPropertyBase::GetExportPath(UDataSourceFilter::StaticClass()) } };
 	AssetRegistryModule.Get().GetAssetsByTagValues(TagValues, Filters);
 
 	for (const FAssetData& AssetData : Filters)
diff --git a/Engine/Plugins/Editor/AssetManagerEditor/Source/AssetManagerEditor/Private/AssetManagerEditorModule.cpp b/Engine/Plugins/Editor/AssetManagerEditor/Source/AssetManagerEditor/Private/AssetManagerEditorModule.cpp
--- a/Engine/Plugins/Editor/AssetManagerEditor/Source/AssetManagerEditor/Private/AssetManagerEditorModule.cpp
+++ b/Engine/Plugins/Editor/AssetManagerEditor/Source/AssetManagerEditor/Private/AssetManagerEditorModule.cpp
@@ -261,14 +261,17 @@ FAssetData IAssetManagerEditorModule::CreateFakeAssetDataFromPrimaryAssetId(cons
 {
 	FString PackageNameString = PrimaryAssetFakeAssetDataPackagePath.ToString() / PrimaryAssetId.PrimaryAssetType.ToString();
 
-	return FAssetData(*PackageNameString, PrimaryAssetFakeAssetDataPackagePath, PrimaryAssetId.PrimaryAssetName, PrimaryAssetId.PrimaryAssetType);
+	// Need to make sure the package part of FTopLevelAssetPath is set otherwise it's gonna be invalid
+	FTopLevelAssetPath FakeAssetClass(PrimaryAssetId.PrimaryAssetType, PrimaryAssetId.PrimaryAssetType);
+	return FAssetData(*PackageNameString, PrimaryAssetFakeAssetDataPackagePath, PrimaryAssetId.PrimaryAssetName, FakeAssetClass);
 }
 
 FPrimaryAssetId IAssetManagerEditorModule::ExtractPrimaryAssetIdFromFakeAssetData(const FAssetData& InAssetData)
 {
 	if (InAssetData.PackagePath == PrimaryAssetFakeAssetDataPackagePath)
 	{
-		return FPrimaryAssetId(InAssetData.AssetClass, InAssetData.AssetName);
+		// See how CreateFakeAssetDataFromPrimaryAssetId stores the asset type inside of FTopLevelAssetPath
+		return FPrimaryAssetId(InAssetData.AssetClassPath.GetAssetName(), InAssetData.AssetName);
 	}
 	return FPrimaryAssetId();
 }
@@ -334,10 +337,10 @@ private:
 
 	//Prints all dependency chains from the PackageName to any dependency of one of the given class names.
 	//If the package name is a path rather than a package, then it will do this for each package in the path.
-	void FindClassDependencies(FName PackagePath, const TArray<FName>& TargetClasses, UE::AssetRegistry::EDependencyQuery RequiredDependencyFlags);
+	void FindClassDependencies(FName PackagePath, const TArray<FTopLevelAssetPath>& TargetClasses, UE::AssetRegistry::EDependencyQuery RequiredDependencyFlags);
 
 	bool GetPackageDependencyChain(FName SourcePackage, FName TargetPackage, TArray<FName>& VisitedPackages, TArray<FName>& OutDependencyChain, UE::AssetRegistry::EDependencyQuery RequiredFlags);
-	void GetPackageDependenciesPerClass(FName SourcePackage, const TArray<FName>& TargetClasses, TArray<FName>& VisitedPackages, TArray<FName>& OutDependentPackages, UE::AssetRegistry::EDependencyQuery RequiredDependencyFlags);
+	void GetPackageDependenciesPerClass(FName SourcePackage, const TArray<FTopLevelAssetPath>& TargetClasses, TArray<FName>& VisitedPackages, TArray<FName>& OutDependentPackages, UE::AssetRegistry::EDependencyQuery RequiredDependencyFlags);
 
 	void LogAssetsWithMultipleLabels();
 	bool CreateOrEmptyCollection(FName CollectionName, ECollectionShareType::Type ShareType);
@@ -1905,12 +1908,19 @@ void FAssetManagerEditorModule::PerformDependencyClassConsoleCommand(const TArra
 	UE::AssetRegistry::EDependencyQuery RequiredDependencyFlags = UE::AssetRegistry::EDependencyQuery::NoRequirements;
 
 	FName SourcePackagePath = FName(*Args[0].ToLower());
-	TArray<FName> TargetClasses;
+	TArray<FTopLevelAssetPath> TargetClasses;
 	for (int32 i = 1; i < Args.Num(); ++i)
 	{
 		if (!GetDependencyTypeArg(Args[i], RequiredDependencyFlags))
 		{
-			TargetClasses.AddUnique(FName(*Args[i]));
+			if (!FPackageName::IsShortPackageName(Args[i]))
+			{
+				TargetClasses.AddUnique(FTopLevelAssetPath(Args[i]));
+			}
+			else
+			{
+				UE_LOG(LogClass, Warning, TEXT("Short class names are not supported: %s"), *Args[i]);
+			}
 		}
 	}
 
@@ -1978,7 +1988,7 @@ bool FAssetManagerEditorModule::GetPackageDependencyChain(FName SourcePackage, F
 	return false;
 }
 
-void FAssetManagerEditorModule::GetPackageDependenciesPerClass(FName SourcePackage, const TArray<FName>& TargetClasses, TArray<FName>& VisitedPackages, TArray<FName>& OutDependentPackages, UE::AssetRegistry::EDependencyQuery RequiredDependencyFlags)
+void FAssetManagerEditorModule::GetPackageDependenciesPerClass(FName SourcePackage, const TArray<FTopLevelAssetPath>& TargetClasses, TArray<FName>& VisitedPackages, TArray<FName>& OutDependentPackages, UE::AssetRegistry::EDependencyQuery RequiredDependencyFlags)
 {
 	//avoid crashing from circular dependencies.
 	if (VisitedPackages.Contains(SourcePackage))
@@ -2003,7 +2013,7 @@ void FAssetManagerEditorModule::GetPackageDependenciesPerClass(FName SourcePacka
 
 	FARFilter Filter;
 	Filter.PackageNames.Add(SourcePackage);
-	Filter.ClassNames = TargetClasses;
+	Filter.ClassPaths = TargetClasses;
 	Filter.bIncludeOnlyOnDiskAssets = true;
 
 	TArray<FAssetData> PackageAssets;
@@ -2059,7 +2069,7 @@ void FAssetManagerEditorModule::FindReferenceChains(FName TargetPackageName, FNa
 	}
 }
 
-void FAssetManagerEditorModule::FindClassDependencies(FName SourcePackageName, const TArray<FName>& TargetClasses, UE::AssetRegistry::EDependencyQuery RequiredDependencyFlags)
+void FAssetManagerEditorModule::FindClassDependencies(FName SourcePackageName, const TArray<FTopLevelAssetPath>& TargetClasses, UE::AssetRegistry::EDependencyQuery RequiredDependencyFlags)
 {
 	TArray<FAssetData> PackageAssets;
 	if (!AssetRegistry->GetAssetsByPackageName(SourcePackageName, PackageAssets))
@@ -2091,16 +2101,16 @@ void FAssetManagerEditorModule::FindClassDependencies(FName SourcePackageName, c
 
 			FARFilter Filter;
 			Filter.PackageNames.Add(DependencyPackage);
-			Filter.ClassNames = TargetClasses;
+			Filter.ClassPaths = TargetClasses;
 			Filter.bIncludeOnlyOnDiskAssets = true;
 
 			if (AssetRegistry->GetAssets(Filter, DepAssets))
 			{
 				for (const FAssetData& DepAsset : DepAssets)
 				{
-					if (TargetClasses.Contains(DepAsset.AssetClass))
+					if (TargetClasses.Contains(DepAsset.AssetClassPath))
 					{
-						UE_LOG(LogAssetManagerEditor, Log, TEXT("Asset: %s class: %s"), *DepAsset.AssetName.ToString(), *DepAsset.AssetClass.ToString());
+						UE_LOG(LogAssetManagerEditor, Log, TEXT("Asset: %s class: %s"), *DepAsset.AssetName.ToString(), *DepAsset.AssetClassPath.ToString());
 					}
 				}
 			}
diff --git a/Engine/Plugins/Editor/AssetManagerEditor/Source/AssetManagerEditor/Private/ReferenceViewer/EdGraphNode_Reference.cpp b/Engine/Plugins/Editor/AssetManagerEditor/Source/AssetManagerEditor/Private/ReferenceViewer/EdGraphNode_Reference.cpp
--- a/Engine/Plugins/Editor/AssetManagerEditor/Source/AssetManagerEditor/Private/ReferenceViewer/EdGraphNode_Reference.cpp
+++ b/Engine/Plugins/Editor/AssetManagerEditor/Source/AssetManagerEditor/Private/ReferenceViewer/EdGraphNode_Reference.cpp
@@ -38,7 +38,7 @@ void UEdGraphNode_Reference::SetupReferenceNode(const FIntPoint& NodeLoc, const
 	Identifiers = NewIdentifiers;
 	const FAssetIdentifier& First = NewIdentifiers[0];
 	FString MainAssetName = InAssetData.AssetName.ToString();
-	FString AssetTypeName = InAssetData.AssetClass.ToString();
+	FString AssetTypeName = InAssetData.AssetClassPath.ToString();
 
 	FAssetToolsModule& AssetToolsModule = FModuleManager::LoadModuleChecked<FAssetToolsModule>(TEXT("AssetTools"));	
 	if (UClass* AssetClass = InAssetData.GetClass())
@@ -49,7 +49,7 @@ void UEdGraphNode_Reference::SetupReferenceNode(const FIntPoint& NodeLoc, const
 			AssetTypeColor = AssetTypeActions.Pin()->GetTypeColor();
 		}
 	}
-	AssetBrush = FSlateIcon("EditorStyle", FName( *("ClassIcon." + InAssetData.AssetClass.ToString())));
+	AssetBrush = FSlateIcon("EditorStyle", FName( *("ClassIcon." + InAssetData.AssetClassPath.ToString())));
 
 	bIsCollapsed = false;
 	bIsPackage = true;
@@ -246,7 +246,7 @@ void UEdGraphNode_Reference::CacheAssetData(const FAssetData& AssetData)
 			{
 				if ( PackageNameStr.StartsWith(TEXT("/Script")) )
 				{
-					CachedAssetData.AssetClass = FName(TEXT("Code"));
+					CachedAssetData.AssetClassPath = FTopLevelAssetPath(TEXT("/EdGraphNode_Reference"), TEXT("Code"));
 				}
 				else
 				{
@@ -254,14 +254,14 @@ void UEdGraphNode_Reference::CacheAssetData(const FAssetData& AssetData)
 					const bool bIsMapPackage = FPlatformFileManager::Get().GetPlatformFile().FileExists(*PotentiallyMapFilename);
 					if ( bIsMapPackage )
 					{
-						CachedAssetData.AssetClass = FName(TEXT("World"));
+						CachedAssetData.AssetClassPath = TEXT("/Script/Engine.World");
 					}
 				}
 			}
 		}
 		else
 		{
-			CachedAssetData.AssetClass = FName(TEXT("Multiple Nodes"));
+			CachedAssetData.AssetClassPath = FTopLevelAssetPath(TEXT("/EdGraphNode_Reference"), TEXT("Multiple Nodes"));
 		}
 	}
 
diff --git a/Engine/Plugins/Editor/AssetManagerEditor/Source/AssetManagerEditor/Private/SAssetAuditBrowser.cpp b/Engine/Plugins/Editor/AssetManagerEditor/Source/AssetManagerEditor/Private/SAssetAuditBrowser.cpp
--- a/Engine/Plugins/Editor/AssetManagerEditor/Source/AssetManagerEditor/Private/SAssetAuditBrowser.cpp
+++ b/Engine/Plugins/Editor/AssetManagerEditor/Source/AssetManagerEditor/Private/SAssetAuditBrowser.cpp
@@ -798,17 +798,17 @@ void SAssetAuditBrowser::AddAssetsOfClass(UClass* AssetClass)
 	{
 		TArray<FAssetData> FoundData;
 		FARFilter AssetFilter;
-		TSet<FName> DerivedClassNames;
+		TSet<FTopLevelAssetPath> DerivedClassNames;
 		bool bAssetClassIsBlueprint = AssetClass->IsChildOf(UBlueprintCore::StaticClass());
 
-		AssetFilter.ClassNames.Add(AssetClass->GetFName());
+		AssetFilter.ClassPaths.Add(AssetClass->GetClassPathName());
 		AssetFilter.bRecursiveClasses = true;
 
 		if (!bAssetClassIsBlueprint)
 		{
 			// If we didn't search specifically for a Blueprint class, find the derived class and look for all blueprints of those derived classes
-			AssetRegistry->GetDerivedClassNames(AssetFilter.ClassNames, TSet<FName>(), DerivedClassNames);
-			AssetFilter.ClassNames.Add(UBlueprintCore::StaticClass()->GetFName());
+			AssetRegistry->GetDerivedClassNames(AssetFilter.ClassPaths, TSet<FTopLevelAssetPath>(), DerivedClassNames);
+			AssetFilter.ClassPaths.Add(UBlueprintCore::StaticClass()->GetClassPathName());
 		}
 		
 		if (AssetRegistry->GetAssets(AssetFilter, FoundData) && FoundData.Num() > 0)
@@ -817,7 +817,7 @@ void SAssetAuditBrowser::AddAssetsOfClass(UClass* AssetClass)
 
 			for (FAssetData& AssetData : FoundData)
 			{
-				if (!bAssetClassIsBlueprint && AssetData.AssetClass.ToString().EndsWith(TEXT("Blueprint")))
+				if (!bAssetClassIsBlueprint && AssetData.AssetClassPath.GetAssetName().ToString().EndsWith(TEXT("Blueprint")))
 				{
 					// This is a blueprint but the filter type wasn't, check the derived class list
 					if (!AssetManager->IsAssetDataBlueprintOfClassSet(AssetData, DerivedClassNames))
diff --git a/Engine/Plugins/Editor/AssetManagerEditor/Source/AssetManagerEditor/Private/SSizeMap.cpp b/Engine/Plugins/Editor/AssetManagerEditor/Source/AssetManagerEditor/Private/SSizeMap.cpp
--- a/Engine/Plugins/Editor/AssetManagerEditor/Source/AssetManagerEditor/Private/SSizeMap.cpp
+++ b/Engine/Plugins/Editor/AssetManagerEditor/Source/AssetManagerEditor/Private/SSizeMap.cpp
@@ -507,7 +507,7 @@ void SSizeMap::GatherDependenciesRecursively(TSharedPtr<FAssetThumbnailPool>& In
 			if (AssetPackageName != NAME_None)
 			{
 				NodeSizeMapData.AssetData.AssetName = AssetPackageName;
-				NodeSizeMapData.AssetData.AssetClass = FName(*LOCTEXT("MissingAsset", "MISSING!").ToString());
+				NodeSizeMapData.AssetData.AssetClassPath = FTopLevelAssetPath(TEXT("/None"), *LOCTEXT("MissingAsset", "MISSING!").ToString());
 
 				const FString AssetPathString = AssetPackageNameString + TEXT(".") + FPackageName::GetLongPackageAssetName(AssetPackageNameString);
 				FAssetData FoundData = CurrentRegistrySource->GetAssetByObjectPath(FName(*AssetPathString));
@@ -679,7 +679,7 @@ void SSizeMap::FinalizeNodesRecursively(TSharedPtr<FTreeMapNodeData>& Node, cons
 			// This has the side effect of not showing a class icon for assets that don't have a proper thumbnail image available
 			bool bIsClassType = false;
 			const UClass* ThumbnailClass = FClassIconFinder::GetIconClassForAssetData(NodeSizeMapData.AssetData, &bIsClassType);
-			const FName DefaultThumbnail = (bIsClassType) ? NAME_None : FName(*FString::Printf(TEXT("ClassThumbnail.%s"), *NodeSizeMapData.AssetData.AssetClass.ToString()));
+			const FName DefaultThumbnail = (bIsClassType) ? NAME_None : FName(*FString::Printf(TEXT("ClassThumbnail.%s"), *NodeSizeMapData.AssetData.AssetClassPath.GetAssetName().ToString()));
 			DefaultThumbnailSlateBrush = FClassIconFinder::FindThumbnailForClass(ThumbnailClass, DefaultThumbnail);
 
 			// @todo sizemap urgent: Actually implement rendered thumbnail support, not just class-based background images
@@ -724,7 +724,7 @@ void SSizeMap::FinalizeNodesRecursively(TSharedPtr<FTreeMapNodeData>& Node, cons
 				// "Asset name"
 				// "Asset type"
 				Node->Name = NodeSizeMapData.AssetData.AssetName.ToString();
-				Node->Name2 = NodeSizeMapData.AssetData.AssetClass.ToString();
+				Node->Name2 = NodeSizeMapData.AssetData.AssetClassPath.ToString();
 			}
 		}
 		else
@@ -743,7 +743,7 @@ void SSizeMap::FinalizeNodesRecursively(TSharedPtr<FTreeMapNodeData>& Node, cons
 				// "Asset name (asset type, size)"
 				Node->Name = FString::Printf(TEXT("%s  (%s, %s)"),
 					*NodeSizeMapData.AssetData.AssetName.ToString(),
-					*NodeSizeMapData.AssetData.AssetClass.ToString(),
+					*NodeSizeMapData.AssetData.AssetClassPath.ToString(),
 					*SizeMapInternals::MakeBestSizeString(SubtreeSize + NodeSizeMapData.AssetSize, !bAnyUnknownSizesInSubtree && NodeSizeMapData.bHasKnownSize));
 			}
 
@@ -761,7 +761,7 @@ void SSizeMap::FinalizeNodesRecursively(TSharedPtr<FTreeMapNodeData>& Node, cons
 				// "*SELF*"
 				// "Asset type"
 				ChildSelfTreeMapNode->Name = LOCTEXT("SelfNodeLabel", "*SELF*").ToString();
-				ChildSelfTreeMapNode->Name2 = NodeSizeMapData.AssetData.AssetClass.ToString();
+				ChildSelfTreeMapNode->Name2 = NodeSizeMapData.AssetData.AssetClassPath.ToString();
 
 				ChildSelfTreeMapNode->CenterText = SizeMapInternals::MakeBestSizeString(NodeSizeMapData.AssetSize, NodeSizeMapData.bHasKnownSize);
 				ChildSelfTreeMapNode->Size = NodeSizeMapData.AssetSize;
diff --git a/Engine/Plugins/Editor/AssetRegistryExport/Source/Private/AssetRegistryExportCommandlet.cpp b/Engine/Plugins/Editor/AssetRegistryExport/Source/Private/AssetRegistryExportCommandlet.cpp
--- a/Engine/Plugins/Editor/AssetRegistryExport/Source/Private/AssetRegistryExportCommandlet.cpp
+++ b/Engine/Plugins/Editor/AssetRegistryExport/Source/Private/AssetRegistryExportCommandlet.cpp
@@ -109,7 +109,7 @@ int32 UAssetRegistryExportCommandlet::Main(const FString& CmdLineParams)
 
 	struct FClassInfo
 	{
-		FName Name;
+		FTopLevelAssetPath Name;
 		TSet<FName> Tags;
 
 		// TagName -> Name that has been sanitized for use as a column name.
@@ -121,15 +121,15 @@ int32 UAssetRegistryExportCommandlet::Main(const FString& CmdLineParams)
 		FSQLitePreparedStatement InsertStatement;
 	};
 
-	TMap<FName, FClassInfo> ClassInfos;
+	TMap<FTopLevelAssetPath, FClassInfo> ClassInfos;
 	int32 AssetCount = 0;
 
 	// Find the tags for each class. Assets might not fill all tags that might exist on a
 	// given class, so we have to get the union of all possibles.
 	AssetRegistry.EnumerateAllAssets(TSet<FName>(), [&](const FAssetData& AssetData)
 	{
-		FClassInfo& ClassInfo = ClassInfos.FindOrAdd(AssetData.AssetClass);
-		ClassInfo.Name = AssetData.AssetClass;
+		FClassInfo& ClassInfo = ClassInfos.FindOrAdd(AssetData.AssetClassPath);
+		ClassInfo.Name = AssetData.AssetClassPath;
 
 		AssetData.TagsAndValues.ForEach([&ClassInfo](TPair<FName, FAssetTagValueRef> Pair)
 		{
@@ -160,7 +160,7 @@ int32 UAssetRegistryExportCommandlet::Main(const FString& CmdLineParams)
 	//
 	// Create the tables and insert statements for each class
 	//
-	for (TPair<FName, FClassInfo>& Pair : ClassInfos)
+	for (TPair<FTopLevelAssetPath, FClassInfo>& Pair : ClassInfos)
 	{
 		FClassInfo& Class = Pair.Value;
 
@@ -239,7 +239,7 @@ int32 UAssetRegistryExportCommandlet::Main(const FString& CmdLineParams)
 			GlobalInsertStatement.ClearBindings();
 
 			GlobalInsertStatement.SetBindingValueByIndex(1, AssetData.AssetName.ToString());
-			GlobalInsertStatement.SetBindingValueByIndex(2, AssetData.AssetClass.ToString());
+			GlobalInsertStatement.SetBindingValueByIndex(2, AssetData.AssetClassPath.ToString());
 			GlobalInsertStatement.SetBindingValueByIndex(3, AssetObjectPath);
 
 			// If this assets has size information then add it to the global assets list as well.
@@ -259,12 +259,12 @@ int32 UAssetRegistryExportCommandlet::Main(const FString& CmdLineParams)
 
 		// Add to our class's table.
 		{
-			FClassInfo& Class = ClassInfos[AssetData.AssetClass];
+			FClassInfo& Class = ClassInfos[AssetData.AssetClassPath];
 
 			Class.InsertStatement.ClearBindings();
 
 			Class.InsertStatement.SetBindingValueByIndex(1, AssetData.AssetName.ToString());
-			Class.InsertStatement.SetBindingValueByIndex(2, AssetData.AssetClass.ToString());
+			Class.InsertStatement.SetBindingValueByIndex(2, AssetData.AssetClassPath.ToString());
 			Class.InsertStatement.SetBindingValueByIndex(3, AssetObjectPath);
 
 			// This will leave any tag values unset (NULL) in the insertion, which is what we want.
@@ -309,7 +309,7 @@ int32 UAssetRegistryExportCommandlet::Main(const FString& CmdLineParams)
 	GlobalInsertStatement.Destroy();
 	CommitTransactionStatement.Destroy();
 	BeginTransactionStatement.Destroy();
-	for (TPair<FName, FClassInfo>& Pair : ClassInfos)
+	for (TPair<FTopLevelAssetPath, FClassInfo>& Pair : ClassInfos)
 	{
 		Pair.Value.InsertStatement.Destroy();
 	}
diff --git a/Engine/Plugins/Editor/AssetSearch/Source/Private/AssetSearchDatabase.cpp b/Engine/Plugins/Editor/AssetSearch/Source/Private/AssetSearchDatabase.cpp
--- a/Engine/Plugins/Editor/AssetSearch/Source/Private/AssetSearchDatabase.cpp
+++ b/Engine/Plugins/Editor/AssetSearch/Source/Private/AssetSearchDatabase.cpp
@@ -177,7 +177,7 @@ public:
 	public: bool AddSearchRecord(const FAssetData& InAssetData, const FString& IndexedJson, const FString& IndexedJsonHash)
 	{
 		FString AssetObjectPath = InAssetData.ObjectPath.ToString();
-		if (Statement_AddAssetToAssetTable.BindAndExecute(InAssetData.AssetName.ToString(), InAssetData.AssetClass.ToString(), AssetObjectPath, IndexedJsonHash))
+		if (Statement_AddAssetToAssetTable.BindAndExecute(InAssetData.AssetName.ToString(), InAssetData.AssetClassPath.ToString(), AssetObjectPath, IndexedJsonHash))
 		{
 			int64 AssetId = Database.GetLastInsertRowId();
 
@@ -333,8 +333,10 @@ public:
 		return Statement_SearchAssets.BindAndExecute(Q, [QueryText, &InCallback](const FSearchAssets& InStatement)
 		{
 			FSearchRecord Result;
-			if (InStatement.GetColumnValues(Result.AssetName, Result.AssetClass, Result.AssetPath))
+			FString ResultClassPath;
+			if (InStatement.GetColumnValues(Result.AssetName, ResultClassPath, Result.AssetPath))
 			{
+				Result.AssetClass = FTopLevelAssetPath(ResultClassPath);
 				const float WorstCase = Result.AssetName.Len() + QueryText.Len();
 				Result.Score = -50.0f * (1.0f - (Algo::LevenshteinDistance(Result.AssetName.ToLower(), QueryText.ToLower()) / WorstCase));
 				
@@ -372,13 +374,15 @@ public:
 		return Statement_SearchAssetPropertiesFTS.BindAndExecute(Q, [&InCallback](const FSearchAssetPropertiesFTS& InStatement)
 		{
 			FSearchRecord Result;
+			FString ResultClassPath;
 			if (InStatement.GetColumnValues(
-				Result.AssetName, Result.AssetClass, Result.AssetPath,
+				Result.AssetName, ResultClassPath, Result.AssetPath,
 				Result.object_name, Result.object_path, Result.object_native_class,
 				Result.property_name, Result.property_field, Result.property_class,
 				Result.value_text, Result.value_hidden,
 				Result.Score))
 			{
+				Result.AssetClass = FTopLevelAssetPath(ResultClassPath);
 				return InCallback(MoveTemp(Result));
 			}
 			return ESQLitePreparedStatementExecuteRowResult::Error;
diff --git a/Engine/Plugins/Editor/AssetSearch/Source/Private/Providers/AssetRegistrySearchProvider.cpp b/Engine/Plugins/Editor/AssetSearch/Source/Private/Providers/AssetRegistrySearchProvider.cpp
--- a/Engine/Plugins/Editor/AssetSearch/Source/Private/Providers/AssetRegistrySearchProvider.cpp
+++ b/Engine/Plugins/Editor/AssetSearch/Source/Private/Providers/AssetRegistrySearchProvider.cpp
@@ -19,7 +19,7 @@ public:
 	/** Get the source tag for the given asset data and alias, or none if there is no match */
 	FName GetSourceTagFromAlias(const FAssetData& InAssetData, const FName InAlias)
 	{
-		TSharedPtr<TMap<FName, FName>>& AliasToSourceTagMapping = ClassToAliasTagsMapping.FindOrAdd(InAssetData.AssetClass);
+		TSharedPtr<TMap<FName, FName>>& AliasToSourceTagMapping = ClassToAliasTagsMapping.FindOrAdd(InAssetData.AssetClassPath);
 
 		if (!AliasToSourceTagMapping.IsValid())
 		{
@@ -61,7 +61,7 @@ public:
 
 private:
 	/** Mapping from class name -> (alias -> source) */
-	TMap<FName, TSharedPtr<TMap<FName, FName>>> ClassToAliasTagsMapping;
+	TMap<FTopLevelAssetPath, TSharedPtr<TMap<FName, FName>>> ClassToAliasTagsMapping;
 };
 
 /** Expression context to test the given asset data against the current text filter */
@@ -168,7 +168,7 @@ public:
 
 		if (bIncludeClassName)
 		{
-			if (InValue.CompareName(AssetPtr->AssetClass, InTextComparisonMode))
+			if (InValue.CompareFString(AssetPtr->AssetClassPath.ToString(), InTextComparisonMode))
 			{
 				return true;
 			}
@@ -245,7 +245,7 @@ public:
 				return false;
 			}
 
-			const bool bIsMatch = TextFilterUtils::TestBasicStringExpression(AssetPtr->AssetClass, InValue, InTextComparisonMode);
+			const bool bIsMatch = TextFilterUtils::TestBasicStringExpression(AssetPtr->AssetClassPath.ToString(), InValue, InTextComparisonMode);
 			return (InComparisonOperation == ETextFilterComparisonOperation::Equal) ? bIsMatch : !bIsMatch;
 		}
 
@@ -365,7 +365,7 @@ void FAssetRegistrySearchProvider::Search(FSearchQueryPtr SearchQuery)
 			FSearchRecord Record;
 			Record.AssetPath = Asset.ObjectPath.ToString();
 			Record.AssetName = Asset.AssetName.ToString();
-			Record.AssetClass = Asset.AssetClass.ToString();
+			Record.AssetClass = Asset.AssetClassPath;
 
 			const float WorstCase = Record.AssetName.Len() + SearchQuery->QueryText.Len();
 			Record.Score = -50.0f * (1.0f - (Algo::LevenshteinDistance(Record.AssetName.ToLower(), SearchQuery->QueryText.ToLower()) / WorstCase));
diff --git a/Engine/Plugins/Editor/AssetSearch/Source/Private/SearchSerializer.cpp b/Engine/Plugins/Editor/AssetSearch/Source/Private/SearchSerializer.cpp
--- a/Engine/Plugins/Editor/AssetSearch/Source/Private/SearchSerializer.cpp
+++ b/Engine/Plugins/Editor/AssetSearch/Source/Private/SearchSerializer.cpp
@@ -54,7 +54,7 @@ void FSearchSerializer::Initialize(const FAssetData& InAsset)
 
 	//JsonWriter->WriteValue(TEXT("AssetName"), InAsset.AssetName.ToString());
 	//JsonWriter->WriteValue(TEXT("AssetPath"), InAsset.ObjectPath.ToString());
-	//JsonWriter->WriteValue(TEXT("AssetClass"), InAsset.AssetClass.ToString());
+	//JsonWriter->WriteValue(TEXT("AssetClass"), InAsset.AssetClassPath.ToString());
 
 	JsonWriter->WriteObjectStart(TEXT("indexers"));
 }
diff --git a/Engine/Plugins/Editor/AssetSearch/Source/Private/Widgets/SSearchTreeRow.cpp b/Engine/Plugins/Editor/AssetSearch/Source/Private/Widgets/SSearchTreeRow.cpp
--- a/Engine/Plugins/Editor/AssetSearch/Source/Private/Widgets/SSearchTreeRow.cpp
+++ b/Engine/Plugins/Editor/AssetSearch/Source/Private/Widgets/SSearchTreeRow.cpp
@@ -54,7 +54,7 @@ TSharedRef<SWidget> SSearchTreeRow::GenerateWidgetForColumn(const FName& ColumnN
 		{
 			TSharedPtr<FAssetObjectNode> ObjectNode = StaticCastSharedPtr<FAssetObjectNode>(BrowserObject);
 
-			if (UClass* ObjectClass = FindObject<UClass>(ANY_PACKAGE, *ObjectNode->object_native_class, true))
+			if (UClass* ObjectClass = UClass::TryFindTypeSlow<UClass>(*ObjectNode->object_native_class, EFindFirstObjectOptions::ExactClass))
 			{
 				//TSharedPtr<IAssetTypeActions> AssetActions = AssetToolsModule.Get().GetAssetTypeActionsForClass(ObjectClass).Pin();
 
@@ -71,7 +71,7 @@ TSharedRef<SWidget> SSearchTreeRow::GenerateWidgetForColumn(const FName& ColumnN
 		{
 			TSharedPtr<FAssetNode> ObjectNode = StaticCastSharedPtr<FAssetNode>(BrowserObject);
 
-			if (UClass* ObjectClass = FindObject<UClass>(ANY_PACKAGE, *ObjectNode->AssetClass, true))
+			if (UClass* ObjectClass = UClass::TryFindTypeSlow<UClass>(ObjectNode->AssetClass, EFindFirstObjectOptions::ExactClass))
 			{
 				TSharedPtr<IAssetTypeActions> AssetActions = AssetToolsModule.Get().GetAssetTypeActionsForClass(ObjectClass).Pin();
 
diff --git a/Engine/Plugins/Editor/AssetSearch/Source/Private/Widgets/SearchModel.h b/Engine/Plugins/Editor/AssetSearch/Source/Private/Widgets/SearchModel.h
--- a/Engine/Plugins/Editor/AssetSearch/Source/Private/Widgets/SearchModel.h
+++ b/Engine/Plugins/Editor/AssetSearch/Source/Private/Widgets/SearchModel.h
@@ -108,7 +108,7 @@ public:
 	FAssetNode(const FSearchRecord& InResult)
 		: AssetName(InResult.AssetName)
 		, AssetPath(InResult.AssetPath)
-		, AssetClass(InResult.AssetClass)
+		, AssetClass(InResult.AssetClass.ToString())
 	{
 		Append(InResult);
 	}
diff --git a/Engine/Plugins/Editor/AssetSearch/Source/Public/SearchQuery.h b/Engine/Plugins/Editor/AssetSearch/Source/Public/SearchQuery.h
--- a/Engine/Plugins/Editor/AssetSearch/Source/Public/SearchQuery.h
+++ b/Engine/Plugins/Editor/AssetSearch/Source/Public/SearchQuery.h
@@ -10,7 +10,7 @@ struct FSearchRecord
 {
 	FString AssetName;
 	FString AssetPath;
-	FString AssetClass;
+	FTopLevelAssetPath AssetClass;
 
 	FString object_name;
 	FString object_path;
diff --git a/Engine/Plugins/Editor/BlueprintHeaderView/Source/BlueprintHeaderView/Private/SBlueprintHeaderView.cpp b/Engine/Plugins/Editor/BlueprintHeaderView/Source/BlueprintHeaderView/Private/SBlueprintHeaderView.cpp
--- a/Engine/Plugins/Editor/BlueprintHeaderView/Source/BlueprintHeaderView/Private/SBlueprintHeaderView.cpp
+++ b/Engine/Plugins/Editor/BlueprintHeaderView/Source/BlueprintHeaderView/Private/SBlueprintHeaderView.cpp
@@ -396,7 +396,7 @@ TSharedRef<SWidget> SBlueprintHeaderView::GetClassPickerMenuContent()
 	FAssetPickerConfig AssetPickerConfig;
 	AssetPickerConfig.SelectionMode = ESelectionMode::Single;
 	AssetPickerConfig.OnAssetSelected = FOnAssetSelected::CreateSP(this, &SBlueprintHeaderView::OnAssetSelected);
-	AssetPickerConfig.Filter.ClassNames.Add(UBlueprint::StaticClass()->GetFName());
+	AssetPickerConfig.Filter.ClassPaths.Add(UBlueprint::StaticClass()->GetClassPathName());
 	AssetPickerConfig.InitialAssetViewType = EAssetViewType::List;
 	AssetPickerConfig.SaveSettingsName = TEXT("HeaderPreviewClassPickerSettings");
 	TSharedRef<SWidget> AssetPickerWidget = ContentBrowserModule.Get().CreateAssetPicker(AssetPickerConfig);
diff --git a/Engine/Plugins/Editor/ConsoleVariablesEditor/Source/ConsoleVariablesEditor/Private/Views/MainPanel/SConsoleVariablesEditorMainPanel.cpp b/Engine/Plugins/Editor/ConsoleVariablesEditor/Source/ConsoleVariablesEditor/Private/Views/MainPanel/SConsoleVariablesEditorMainPanel.cpp
--- a/Engine/Plugins/Editor/ConsoleVariablesEditor/Source/ConsoleVariablesEditor/Private/Views/MainPanel/SConsoleVariablesEditorMainPanel.cpp
+++ b/Engine/Plugins/Editor/ConsoleVariablesEditor/Source/ConsoleVariablesEditor/Private/Views/MainPanel/SConsoleVariablesEditorMainPanel.cpp
@@ -279,7 +279,7 @@ TSharedRef<SWidget> SConsoleVariablesEditorMainPanel::OnGeneratePresetsMenu()
 		AssetPickerConfig.bForceShowPluginContent = false;
 
 		AssetPickerConfig.AssetShowWarningText = LOCTEXT("NoPresets_Warning", "No Presets Found");
-		AssetPickerConfig.Filter.ClassNames.Add(UConsoleVariablesAsset::StaticClass()->GetFName());
+		AssetPickerConfig.Filter.ClassPaths.Add(UConsoleVariablesAsset::StaticClass()->GetClassPathName());
 		AssetPickerConfig.Filter.bRecursiveClasses = true;
 		AssetPickerConfig.OnAssetSelected =
 			FOnAssetSelected::CreateRaw(MainPanel.Pin().Get(), &FConsoleVariablesEditorMainPanel::ImportPreset);
diff --git a/Engine/Plugins/Editor/ContentBrowser/ContentBrowserAliasDataSoure/Source/ContentBrowserAliasDataSource/Private/ContentBrowserAliasDataSource.cpp b/Engine/Plugins/Editor/ContentBrowser/ContentBrowserAliasDataSoure/Source/ContentBrowserAliasDataSource/Private/ContentBrowserAliasDataSource.cpp
--- a/Engine/Plugins/Editor/ContentBrowser/ContentBrowserAliasDataSoure/Source/ContentBrowserAliasDataSource/Private/ContentBrowserAliasDataSource.cpp
+++ b/Engine/Plugins/Editor/ContentBrowser/ContentBrowserAliasDataSoure/Source/ContentBrowserAliasDataSource/Private/ContentBrowserAliasDataSource.cpp
@@ -274,7 +274,7 @@ bool UContentBrowserAliasDataSource::DoesAliasPassFilter(const FAliasData& Alias
 {
 	// Create a fake asset data using the alias path instead of the asset's original path
 	// AssetRegistry->IsAssetIncludedByFilter is effectively a static function and does not actually use AssetRegistry data
-	FAssetData AliasAssetData(AliasData.AliasName, AliasData.PackagePath, AliasData.AliasName, AliasData.AssetData.AssetClass, AliasData.AssetData.TagsAndValues.CopyMap());
+	FAssetData AliasAssetData(AliasData.AliasName, AliasData.PackagePath, AliasData.AliasName, AliasData.AssetData.AssetClassPath, AliasData.AssetData.TagsAndValues.CopyMap());
 
 	return (Filter.InclusiveFilter.IsEmpty() || AssetRegistry->IsAssetIncludedByFilter(AliasAssetData, Filter.InclusiveFilter)) // Passes Inclusive
 		&& (Filter.ExclusiveFilter.IsEmpty() || !AssetRegistry->IsAssetIncludedByFilter(AliasAssetData, Filter.ExclusiveFilter)); // Passes Exclusive
@@ -723,7 +723,7 @@ bool UContentBrowserAliasDataSource::Legacy_TryConvertPackagePathToVirtualPath(c
 
 bool UContentBrowserAliasDataSource::Legacy_TryConvertAssetDataToVirtualPath(const FAssetData& InAssetData, const bool InUseFolderPaths, FName& OutPath)
 {
-	return InAssetData.AssetClass != NAME_Class // Ignore legacy class items
+	return InAssetData.AssetClassPath != FTopLevelAssetPath(TEXT("/Script/CoreUObject"), NAME_Class) // Ignore legacy class items
 		&& TryConvertInternalPathToVirtual(InUseFolderPaths ? InAssetData.PackagePath : InAssetData.ObjectPath, OutPath);
 }
 
diff --git a/Engine/Plugins/Editor/ContentBrowser/ContentBrowserAssetDataSource/Source/ContentBrowserAssetDataSource/Private/AssetFileContextMenu.cpp b/Engine/Plugins/Editor/ContentBrowser/ContentBrowserAssetDataSource/Source/ContentBrowserAssetDataSource/Private/AssetFileContextMenu.cpp
--- a/Engine/Plugins/Editor/ContentBrowser/ContentBrowserAssetDataSource/Source/ContentBrowserAssetDataSource/Private/AssetFileContextMenu.cpp
+++ b/Engine/Plugins/Editor/ContentBrowser/ContentBrowserAssetDataSource/Source/ContentBrowserAssetDataSource/Private/AssetFileContextMenu.cpp
@@ -471,7 +471,7 @@ void FAssetFileContextMenu::MakeAssetActionsSubMenu(UToolMenu* Menu)
 		// Materials can't be bulk edited currently as they require very special handling because of their dependencies with the rendering thread, and we'd have to hack the property matrix too much.
 		for (auto& Asset : SelectedAssets)
 		{
-			if (Asset.AssetClass == UMaterial::StaticClass()->GetFName() || Asset.AssetClass == UMaterialInstanceConstant::StaticClass()->GetFName() || Asset.AssetClass == UMaterialFunction::StaticClass()->GetFName() || Asset.AssetClass == UMaterialFunctionInstance::StaticClass()->GetFName())
+			if (Asset.AssetClassPath == UMaterial::StaticClass()->GetClassPathName() || Asset.AssetClassPath == UMaterialInstanceConstant::StaticClass()->GetClassPathName() || Asset.AssetClassPath == UMaterialFunction::StaticClass()->GetClassPathName() || Asset.AssetClassPath == UMaterialFunctionInstance::StaticClass()->GetClassPathName())
 			{
 				bCanUsePropertyMatrix = false;
 				break;
@@ -1520,7 +1520,7 @@ void FAssetFileContextMenu::GetSelectedAssets(TArray<UObject*>& Assets, bool Ski
 {
 	for (int32 AssetIdx = 0; AssetIdx < SelectedAssets.Num(); ++AssetIdx)
 	{
-		if (SkipRedirectors && (SelectedAssets[AssetIdx].AssetClass == UObjectRedirector::StaticClass()->GetFName()))
+		if (SkipRedirectors && (SelectedAssets[AssetIdx].AssetClassPath == UObjectRedirector::StaticClass()->GetClassPathName()))
 		{
 			// Don't operate on Redirectors
 			continue;
@@ -1782,13 +1782,13 @@ void FAssetFileContextMenu::ExecuteReload()
 		{
 			const FAssetData& AssetData = *AssetIt;
 
-			if (AssetData.AssetClass == UObjectRedirector::StaticClass()->GetFName())
+			if (AssetData.AssetClassPath == UObjectRedirector::StaticClass()->GetClassPathName())
 			{
 				// Don't operate on Redirectors
 				continue;
 			}
 
-			if (AssetData.AssetClass == UUserDefinedStruct::StaticClass()->GetFName())
+			if (AssetData.AssetClassPath == UUserDefinedStruct::StaticClass()->GetClassPathName())
 			{
 				FNotificationInfo Notification(LOCTEXT("CannotReloadUserStruct", "User created structures cannot be safely reloaded."));
 				Notification.ExpireDuration = 3.0f;
@@ -1796,7 +1796,7 @@ void FAssetFileContextMenu::ExecuteReload()
 				continue;
 			}
 
-			if (AssetData.AssetClass == UUserDefinedEnum::StaticClass()->GetFName())
+			if (AssetData.AssetClassPath == UUserDefinedEnum::StaticClass()->GetClassPathName())
 			{
 				FNotificationInfo Notification(LOCTEXT("CannotReloadUserEnum", "User created enumerations cannot be safely reloaded."));
 				Notification.ExpireDuration = 3.0f;
@@ -2393,7 +2393,7 @@ bool FAssetFileContextMenu::CanExecuteDiffSelected() const
 		FAssetData const& FirstSelection = SelectedAssets[0];
 		FAssetData const& SecondSelection = SelectedAssets[1];
 
-		bCanDiffSelected = FirstSelection.AssetClass == SecondSelection.AssetClass;
+		bCanDiffSelected = FirstSelection.AssetClassPath == SecondSelection.AssetClassPath;
 	}
 
 	return bCanDiffSelected;
@@ -2436,7 +2436,7 @@ void FAssetFileContextMenu::CacheCanExecuteVars()
 			continue;
 		}
 
-		if ( !bAtLeastOneNonRedirectorSelected && AssetData.AssetClass != UObjectRedirector::StaticClass()->GetFName() )
+		if ( !bAtLeastOneNonRedirectorSelected && AssetData.AssetClassPath != UObjectRedirector::StaticClass()->GetClassPathName() )
 		{
 			bAtLeastOneNonRedirectorSelected = true;
 		}
diff --git a/Engine/Plugins/Editor/ContentBrowser/ContentBrowserAssetDataSource/Source/ContentBrowserAssetDataSource/Private/AssetFolderContextMenu.cpp b/Engine/Plugins/Editor/ContentBrowser/ContentBrowserAssetDataSource/Source/ContentBrowserAssetDataSource/Private/AssetFolderContextMenu.cpp
--- a/Engine/Plugins/Editor/ContentBrowser/ContentBrowserAssetDataSource/Source/ContentBrowserAssetDataSource/Private/AssetFolderContextMenu.cpp
+++ b/Engine/Plugins/Editor/ContentBrowser/ContentBrowserAssetDataSource/Source/ContentBrowserAssetDataSource/Private/AssetFolderContextMenu.cpp
@@ -224,7 +224,7 @@ void FAssetFolderContextMenu::ExecuteFixUpRedirectorsInFolder()
 	for (const auto& Path : SelectedPaths)
 	{
 		Filter.PackagePaths.Emplace(*Path);
-		Filter.ClassNames.Emplace(TEXT("ObjectRedirector"));
+		Filter.ClassPaths.Emplace(FTopLevelAssetPath(TEXT("/Script/CoreUObject"), TEXT("ObjectRedirector")));
 	}
 
 	// Query for a list of assets in the selected paths
diff --git a/Engine/Plugins/Editor/ContentBrowser/ContentBrowserAssetDataSource/Source/ContentBrowserAssetDataSource/Private/AssetPropertyTagCache.cpp b/Engine/Plugins/Editor/ContentBrowser/ContentBrowserAssetDataSource/Source/ContentBrowserAssetDataSource/Private/AssetPropertyTagCache.cpp
--- a/Engine/Plugins/Editor/ContentBrowser/ContentBrowserAssetDataSource/Source/ContentBrowserAssetDataSource/Private/AssetPropertyTagCache.cpp
+++ b/Engine/Plugins/Editor/ContentBrowser/ContentBrowserAssetDataSource/Source/ContentBrowserAssetDataSource/Private/AssetPropertyTagCache.cpp
@@ -13,7 +13,7 @@ FAssetPropertyTagCache& FAssetPropertyTagCache::Get()
 	return AssetPropertyTagCache;
 }
 
-const FAssetPropertyTagCache::FClassPropertyTagCache& FAssetPropertyTagCache::GetCacheForClass(const FName InClassName)
+const FAssetPropertyTagCache::FClassPropertyTagCache& FAssetPropertyTagCache::GetCacheForClass(FTopLevelAssetPath InClassName)
 {
 	TSharedPtr<FClassPropertyTagCache> ClassCache = ClassToCacheMap.FindRef(InClassName);
 	if (!ClassCache)
@@ -22,15 +22,15 @@ const FAssetPropertyTagCache::FClassPropertyTagCache& FAssetPropertyTagCache::Ge
 
 		auto GetAssetClass = [&InClassName]()
 		{
-			UClass* FoundClass = FindObject<UClass>(ANY_PACKAGE, *InClassName.ToString());
+			UClass* FoundClass = FindObject<UClass>(InClassName);
 
 			if (!FoundClass)
 			{
 				// Look for class redirectors
-				FName NewPath = FLinkerLoad::FindNewNameForClass(InClassName, false);
-				if (NewPath != NAME_None)
+				FString NewPath = FLinkerLoad::FindNewPathNameForClass(InClassName.ToString(), false);
+				if (!NewPath.IsEmpty())
 				{
-					FoundClass = FindObject<UClass>(ANY_PACKAGE, *NewPath.ToString());
+					FoundClass = FindObject<UClass>(nullptr, *NewPath);
 				}
 			}
 
