Create a setting to enable / disable supporting the mixing of source and PCB builds. This was causing issues when sharing engine builds with different UGS projects.

diff --git a/Engine/Plugins/Developer/Concert/ConcertApp/MultiUserServer/Source/MultiUserServer/Private/Widgets/Browser/ConcertServerSessionBrowserController.cpp b/Engine/Plugins/Developer/Concert/ConcertApp/MultiUserServer/Source/MultiUserServer/Private/Widgets/Browser/ConcertServerSessionBrowserController.cpp
--- a/Engine/Plugins/Developer/Concert/ConcertApp/MultiUserServer/Source/MultiUserServer/Private/Widgets/Browser/ConcertServerSessionBrowserController.cpp
+++ b/Engine/Plugins/Developer/Concert/ConcertApp/MultiUserServer/Source/MultiUserServer/Private/Widgets/Browser/ConcertServerSessionBrowserController.cpp
@@ -111,7 +111,7 @@ void FConcertServerSessionBrowserController::CreateSession(const FGuid& ServerAd
 		SessionInfo.Settings.ProjectName = ProjectName;
 	}
 	FConcertSessionVersionInfo VersionInfo;
-	VersionInfo.Initialize();
+	VersionInfo.Initialize(false /* bSupportMixedBuildTypes */ );
 	SessionInfo.VersionInfos.Emplace(VersionInfo);
 
 	FAsyncTaskNotification Notification(
diff --git a/Engine/Plugins/Developer/Concert/ConcertMain/Source/Concert/Private/ConcertClient.cpp b/Engine/Plugins/Developer/Concert/ConcertMain/Source/Concert/Private/ConcertClient.cpp
--- a/Engine/Plugins/Developer/Concert/ConcertMain/Source/Concert/Private/ConcertClient.cpp
+++ b/Engine/Plugins/Developer/Concert/ConcertMain/Source/Concert/Private/ConcertClient.cpp
@@ -1174,7 +1174,7 @@ TFuture<EConcertResponseCode> FConcertClient::InternalCreateSession(const FGuid&
 		FConcertAdmin_CreateSessionRequest CreateSessionRequest;
 		CreateSessionRequest.SessionName = CreateSessionArgs.SessionName;
 		CreateSessionRequest.OwnerClientInfo = ClientInfo;
-		CreateSessionRequest.VersionInfo.Initialize();
+		CreateSessionRequest.VersionInfo.Initialize(GetConfiguration()->ClientSettings.bSupportMixedBuildTypes);
 	
 		// Session settings
 		CreateSessionRequest.SessionSettings.Initialize();
@@ -1220,7 +1220,7 @@ TFuture<EConcertResponseCode> FConcertClient::InternalJoinSession(const FGuid& S
 		FConcertAdmin_FindSessionRequest FindSessionRequest;
 		FindSessionRequest.SessionId = SessionId;
 		FindSessionRequest.OwnerClientInfo = ClientInfo;
-		FindSessionRequest.VersionInfo.Initialize();
+		FindSessionRequest.VersionInfo.Initialize(GetConfiguration()->ClientSettings.bSupportMixedBuildTypes);
 
 		// Session settings
 		FindSessionRequest.SessionSettings.Initialize();
@@ -1249,7 +1249,7 @@ TFuture<EConcertResponseCode> FConcertClient::InternalCopySession(const FGuid& S
 	CopySessionRequest.SessionFilter = CopySessionArgs.SessionFilter;
 	CopySessionRequest.bRestoreOnly = bRestoreOnlyConstraint;
 	CopySessionRequest.OwnerClientInfo = ClientInfo;
-	CopySessionRequest.VersionInfo.Initialize();
+	CopySessionRequest.VersionInfo.Initialize(GetConfiguration()->ClientSettings.bSupportMixedBuildTypes);
 
 	// Session settings
 	CopySessionRequest.SessionSettings.Initialize();
diff --git a/Engine/Plugins/Developer/Concert/ConcertMain/Source/Concert/Private/ConcertSettings.cpp b/Engine/Plugins/Developer/Concert/ConcertMain/Source/Concert/Private/ConcertSettings.cpp
--- a/Engine/Plugins/Developer/Concert/ConcertMain/Source/Concert/Private/ConcertSettings.cpp
+++ b/Engine/Plugins/Developer/Concert/ConcertMain/Source/Concert/Private/ConcertSettings.cpp
@@ -59,7 +59,7 @@ UConcertServerConfig::UConcertServerConfig()
 	: bCleanWorkingDir(false)
 	, NumSessionsToKeep(-1)
 {
-	DefaultVersionInfo.Initialize();
+	DefaultVersionInfo.Initialize(false /* bSupportMixedBuildTypes */);
 }
 
 UConcertClientConfig::UConcertClientConfig()
diff --git a/Engine/Plugins/Developer/Concert/ConcertMain/Source/Concert/Private/ConcertVersion.cpp b/Engine/Plugins/Developer/Concert/ConcertMain/Source/Concert/Private/ConcertVersion.cpp
--- a/Engine/Plugins/Developer/Concert/ConcertMain/Source/Concert/Private/ConcertVersion.cpp
+++ b/Engine/Plugins/Developer/Concert/ConcertMain/Source/Concert/Private/ConcertVersion.cpp
@@ -1,14 +1,14 @@
 // Copyright Epic Games, Inc. All Rights Reserved.
 
 #include "ConcertVersion.h"
+#include "ConcertSettings.h"
+
 #include "Misc/EngineVersion.h"
 #include "Modules/BuildVersion.h"
 #include "Serialization/CustomVersion.h"
 
 #define LOCTEXT_NAMESPACE "ConcertVersion"
 
-
-
 namespace ConcertVersionUtil
 {
 
@@ -95,29 +95,31 @@ bool FConcertCustomVersionInfo::Validate(const FConcertCustomVersionInfo& InOthe
 	return ConcertVersionUtil::ValidateVersion(Version, InOther.Version, FText::AsCultureInvariant(FriendlyName.IsNone() ? Key.ToString() : FriendlyName.ToString()), InValidationMode, OutFailureReason);
 }
 
-
-void FConcertSessionVersionInfo::Initialize()
+void FConcertSessionVersionInfo::Initialize(bool bSupportMixedBuildTypes)
 {
 	FileVersion.Initialize();
 	EngineVersion.Initialize(FEngineVersion::Current());
 
-	// For builds synced via UGS, we override the changelist of the engine version with the current build version changelist
-	// as this helps to keep the changelists of programmers and artists/designers in-sync when creating and joining sessions
-	//	eg) CL# 1 is a code change, and CL# 2 is a content change:
-	//	 - A programmer syncing CL# 2 would have an engine version with a CL# of 2 (from building their own editor), and a build version CL# of 2 (from UGS).
-	//	 - An artist/designer syncing CL# 2 would have an engine version with a CL# of 1 (from the pre-built editor), but a build version CL# of 2 (from UGS).
+	if (bSupportMixedBuildTypes)
 	{
-		// Read the default data (rather than the executable specific data), as the default data 
-		// is updated when syncing, but the executable data is only updated when compiling
-		FBuildVersion BuildVersion;
-		if (FBuildVersion::TryRead(FBuildVersion::GetDefaultFileName(), BuildVersion))
+		// For builds synced via UGS, we override the changelist of the engine version with the current build version changelist
+		// as this helps to keep the changelists of programmers and artists/designers in-sync when creating and joining sessions
+		//	eg) CL# 1 is a code change, and CL# 2 is a content change:
+		//	 - A programmer syncing CL# 2 would have an engine version with a CL# of 2 (from building their own editor), and a build version CL# of 2 (from UGS).
+		//	 - An artist/designer syncing CL# 2 would have an engine version with a CL# of 1 (from the pre-built editor), but a build version CL# of 2 (from UGS).
 		{
-			// Only apply the build version if our engine changelist is compatible with the synced build
-			// If this check fails then it likely means that a programmer synced (updating the build version) 
-			// without also compiling their binaries (to update the engine version)
-			if (EngineVersion.Changelist >= (uint32)BuildVersion.CompatibleChangelist)
+			// Read the default data (rather than the executable specific data), as the default data 
+			// is updated when syncing, but the executable data is only updated when compiling
+			FBuildVersion BuildVersion;
+			if (FBuildVersion::TryRead(FBuildVersion::GetDefaultFileName(), BuildVersion))
 			{
-				EngineVersion.Changelist = BuildVersion.Changelist;
+				// Only apply the build version if our engine changelist is compatible with the synced build
+				// If this check fails then it likely means that a programmer synced (updating the build version) 
+				// without also compiling their binaries (to update the engine version)
+				if (EngineVersion.Changelist >= (uint32)BuildVersion.CompatibleChangelist)
+				{
+					EngineVersion.Changelist = BuildVersion.Changelist;
+				}
 			}
 		}
 	}
diff --git a/Engine/Plugins/Developer/Concert/ConcertMain/Source/Concert/Public/ConcertSettings.h b/Engine/Plugins/Developer/Concert/ConcertMain/Source/Concert/Public/ConcertSettings.h
--- a/Engine/Plugins/Developer/Concert/ConcertMain/Source/Concert/Public/ConcertSettings.h
+++ b/Engine/Plugins/Developer/Concert/ConcertMain/Source/Concert/Public/ConcertSettings.h
@@ -250,6 +250,15 @@ struct FConcertClientSettings
 	UPROPERTY(config, EditAnywhere, DisplayName="Reflect Level Visibility to Game", AdvancedDisplay, Category="Client Settings")
 	bool bReflectLevelEditorInGame = false;
 
+	/**
+	 * Enable extended version support when using Multi-user with precompiled and source builds.  When using Unreal Game
+	 * Sync, it is possible to have the same engine CL but different engine version due to content changes.	 This setting
+	 * enables reading engine version CL from the Build.version file produced by UGS to determine engine version
+	 * information when joining a session.	This only applies when you intend to mix precompiled with source builds.
+	*/
+	UPROPERTY(config, EditAnywhere, DisplayName="Support Mixed Build Types", AdvancedDisplay, Category="Client Settings")
+	bool bSupportMixedBuildTypes = false;
+
 	/** Array of tags that can be used for grouping and categorizing. */
 	UPROPERTY(config, EditAnywhere, AdvancedDisplay, Category = "Client Settings")
 	TArray<FName> Tags;
diff --git a/Engine/Plugins/Developer/Concert/ConcertMain/Source/Concert/Public/ConcertVersion.h b/Engine/Plugins/Developer/Concert/ConcertMain/Source/Concert/Public/ConcertVersion.h
--- a/Engine/Plugins/Developer/Concert/ConcertMain/Source/Concert/Public/ConcertVersion.h
+++ b/Engine/Plugins/Developer/Concert/ConcertMain/Source/Concert/Public/ConcertVersion.h
@@ -103,7 +103,7 @@ struct FConcertSessionVersionInfo
 	GENERATED_BODY()
 
 	/** Initialize this version info from the compiled in data */
-	CONCERT_API void Initialize();
+	CONCERT_API void Initialize(bool bSupportMixedBuildTypes);
 
 	/** Validate this version info against another */
 	CONCERT_API bool Validate(const FConcertSessionVersionInfo& InOther, const EConcertVersionValidationMode InValidationMode, FText* OutFailureReason = nullptr) const;
