Crash when user has no permission to delete

diff --git a/Engine/Plugins/Developer/Concert/ConcertMain/Source/Concert/Private/ConcertServer.cpp b/Engine/Plugins/Developer/Concert/ConcertMain/Source/Concert/Private/ConcertServer.cpp
--- a/Engine/Plugins/Developer/Concert/ConcertMain/Source/Concert/Private/ConcertServer.cpp
+++ b/Engine/Plugins/Developer/Concert/ConcertMain/Source/Concert/Private/ConcertServer.cpp
@@ -1299,12 +1299,12 @@
 			const bool bHasPermission = IsRequestFromSessionOwner(*ArchivedSessionInfo, Request.UserName, Request.DeviceName);
 			if (!bHasPermission && (Request.Flags & EBatchSessionDeletionFlags::SkipForbiddenSessions) != EBatchSessionDeletionFlags::Strict)
 			{
-				NotOwnedByClient.Add({ SessionToDelete, ServerSession->GetName() });
+				NotOwnedByClient.Add({ SessionToDelete, ArchivedSessionInfo->SessionName });
 			}
 			else if (!bHasPermission)
 			{
 				OutResponse.Reason = LOCTEXT("Error_BatchDelete_InvalidPerms_NotOwner", "Not the session owner.");
-				UE_LOG(LogConcert, Display, TEXT("User %s failed to delete live session '%s' (Id: %s, Owner: %s, Reason: %s)"), *Request.UserName, *ArchivedSessionInfo->SessionName, *SessionToDelete.ToString(), *ServerSession->GetSessionInfo().OwnerUserName, *OutResponse.Reason.ToString());
+				UE_LOG(LogConcert, Display, TEXT("User %s failed to delete live session '%s' (Id: %s, Owner: %s, Reason: %s)"), *Request.UserName, *ArchivedSessionInfo->SessionName, *SessionToDelete.ToString(), *ArchivedSessionInfo->OwnerUserName, *OutResponse.Reason.ToString());
 				return false;
 			}
 			
