Query remote client message address

diff --git a/Engine/Plugins/Developer/Concert/ConcertMain/Source/Concert/Private/ConcertServer.cpp b/Engine/Plugins/Developer/Concert/ConcertMain/Source/Concert/Private/ConcertServer.cpp
--- a/Engine/Plugins/Developer/Concert/ConcertMain/Source/Concert/Private/ConcertServer.cpp
+++ b/Engine/Plugins/Developer/Concert/ConcertMain/Source/Concert/Private/ConcertServer.cpp
@@ -150,6 +150,15 @@ const FConcertServerInfo& FConcertServer::GetServerInfo() const
 	return ServerInfo;
 }
 
+FMessageAddress FConcertServer::GetRemoteAddress(const FGuid& AdminEndpointId) const
+{
+	if (IsStarted())
+	{
+		return ServerAdminEndpoint->GetRemoteAddress(AdminEndpointId);
+	}
+	return {};
+}
+
 bool FConcertServer::IsStarted() const
 {
 	return ServerAdminEndpoint.IsValid();
diff --git a/Engine/Plugins/Developer/Concert/ConcertMain/Source/Concert/Private/ConcertServer.h b/Engine/Plugins/Developer/Concert/ConcertMain/Source/Concert/Private/ConcertServer.h
--- a/Engine/Plugins/Developer/Concert/ConcertMain/Source/Concert/Private/ConcertServer.h
+++ b/Engine/Plugins/Developer/Concert/ConcertMain/Source/Concert/Private/ConcertServer.h
@@ -25,6 +25,7 @@
 	virtual bool IsConfigured() const override;
 	virtual const UConcertServerConfig* GetConfiguration() const override;
 	virtual const FConcertServerInfo& GetServerInfo() const override;
+	virtual FMessageAddress GetRemoteAddress(const FGuid& AdminEndpointId) const override;
 
 	virtual bool IsStarted() const override;
 	virtual void Startup() override;
diff --git a/Engine/Plugins/Developer/Concert/ConcertMain/Source/Concert/Public/IConcertServer.h b/Engine/Plugins/Developer/Concert/ConcertMain/Source/Concert/Public/IConcertServer.h
--- a/Engine/Plugins/Developer/Concert/ConcertMain/Source/Concert/Public/IConcertServer.h
+++ b/Engine/Plugins/Developer/Concert/ConcertMain/Source/Concert/Public/IConcertServer.h
@@ -38,6 +38,9 @@
 	 */
 	virtual const FConcertServerInfo& GetServerInfo() const = 0;
 
+	/** Gets the address of a remote admin endpoint, i.e. a client that is sending FConcertEndpointDiscoveryEvents. */
+	virtual FMessageAddress GetRemoteAddress(const FGuid& AdminEndpointId) const = 0;
+
 	/**
 	 *	Returns if the server has already been started up.
 	 */
