Expose logged ACKs

diff --git a/Engine/Plugins/Developer/Concert/ConcertMain/Source/Concert/Private/ConcertServer.cpp b/Engine/Plugins/Developer/Concert/ConcertMain/Source/Concert/Private/ConcertServer.cpp
--- a/Engine/Plugins/Developer/Concert/ConcertMain/Source/Concert/Private/ConcertServer.cpp
+++ b/Engine/Plugins/Developer/Concert/ConcertMain/Source/Concert/Private/ConcertServer.cpp
@@ -159,6 +159,11 @@ FMessageAddress FConcertServer::GetRemoteAddress(const FGuid& AdminEndpointId) c
 	return {};
 }
 
+FOnConcertMessageAcknowledgementReceivedFromLocalEndpoint& FConcertServer::OnConcertMessageAcknowledgementReceived()
+{
+	return OnConcertMessageAcknowledgementReceivedFromLocalEndpoint;
+}
+
 bool FConcertServer::IsStarted() const
 {
 	return ServerAdminEndpoint.IsValid();
@@ -178,6 +183,11 @@ void FConcertServer::Startup()
 			});
 		});
 		ServerInfo.AdminEndpointId = ServerAdminEndpoint->GetEndpointContext().EndpointId;
+		ServerAdminEndpoint->OnConcertMessageAcknowledgementReceived().AddLambda(
+		[this](const FConcertEndpointContext& LocalEndpoint, const FConcertEndpointContext& RemoteEndpoint, const TSharedRef<IConcertMessage>& AckedMessage, const FConcertMessageContext& MessageContext)
+			{
+				OnConcertMessageAcknowledgementReceivedFromLocalEndpoint.Broadcast(LocalEndpoint, RemoteEndpoint, AckedMessage, MessageContext);
+			});
 
 		// Make it discoverable
 		ServerAdminEndpoint->SubscribeEventHandler<FConcertAdmin_DiscoverServersEvent>(this, &FConcertServer::HandleDiscoverServersEvent);
@@ -1513,17 +1523,24 @@ TSharedPtr<IConcertServerSession> FConcertServer::CreateLiveSession(const FConce
 		UE_CLOG(LiveSessionInfo.VersionInfos.Num() > 0, LogConcert, Warning, TEXT("Clearing version information when creating session '%s' due to -CONCERTIGNORE. This session will be unversioned!"), *LiveSessionInfo.SessionName);
 		LiveSessionInfo.VersionInfos.Reset();
 	}
-
-	TSharedPtr<FConcertServerSession> LiveSession = MakeShared<FConcertServerSession>(
-		LiveSessionInfo,
-		Settings->ServerSettings,
-		EndpointProvider->CreateLocalEndpoint(LiveSessionInfo.SessionName, Settings->EndpointSettings, [this](const FConcertEndpointContext& Context)
+	
+	TSharedPtr<IConcertLocalEndpoint> SessionEndpoint = EndpointProvider->CreateLocalEndpoint(LiveSessionInfo.SessionName, Settings->EndpointSettings, [this](const FConcertEndpointContext& Context)
 		{
 			return FConcertLogger::CreateLogger(Context, [this](const FConcertLog& Log)
 			{
 				ConcertTransportEvents::OnConcertServerLogEvent().Broadcast(*this, Log);
 			});
-		}),
+		});
+	SessionEndpoint->OnConcertMessageAcknowledgementReceived().AddLambda(
+		[this](const FConcertEndpointContext& LocalEndpoint, const FConcertEndpointContext& RemoteEndpoint, const TSharedRef<IConcertMessage>& AckedMessage, const FConcertMessageContext& MessageContext)
+		{
+			OnConcertMessageAcknowledgementReceivedFromLocalEndpoint.Broadcast(LocalEndpoint, RemoteEndpoint, AckedMessage, MessageContext);
+		});
+
+	TSharedPtr<FConcertServerSession> LiveSession = MakeShared<FConcertServerSession>(
+		LiveSessionInfo,
+		Settings->ServerSettings,
+		SessionEndpoint,
 		InRepository.GetSessionWorkingDir(LiveSessionInfo.SessionId)
 		);
 
diff --git a/Engine/Plugins/Developer/Concert/ConcertMain/Source/Concert/Private/ConcertServer.h b/Engine/Plugins/Developer/Concert/ConcertMain/Source/Concert/Private/ConcertServer.h
--- a/Engine/Plugins/Developer/Concert/ConcertMain/Source/Concert/Private/ConcertServer.h
+++ b/Engine/Plugins/Developer/Concert/ConcertMain/Source/Concert/Private/ConcertServer.h
@@ -26,6 +26,7 @@ public:
 	virtual const UConcertServerConfig* GetConfiguration() const override;
 	virtual const FConcertServerInfo& GetServerInfo() const override;
 	virtual FMessageAddress GetRemoteAddress(const FGuid& AdminEndpointId) const override;
+	virtual FOnConcertMessageAcknowledgementReceivedFromLocalEndpoint& OnConcertMessageAcknowledgementReceived() override;
 
 	virtual bool IsStarted() const override;
 	virtual void Startup() override;
@@ -193,4 +194,6 @@ private:
 
 	/** The root directory containing the session repositories. */
 	FString SessionRepositoryRootDir;
+
+	FOnConcertMessageAcknowledgementReceivedFromLocalEndpoint OnConcertMessageAcknowledgementReceivedFromLocalEndpoint; 
 };
diff --git a/Engine/Plugins/Developer/Concert/ConcertMain/Source/Concert/Public/IConcertServer.h b/Engine/Plugins/Developer/Concert/ConcertMain/Source/Concert/Public/IConcertServer.h
--- a/Engine/Plugins/Developer/Concert/ConcertMain/Source/Concert/Public/IConcertServer.h
+++ b/Engine/Plugins/Developer/Concert/ConcertMain/Source/Concert/Public/IConcertServer.h
@@ -3,6 +3,7 @@
 #pragma once
 
 #include "CoreMinimal.h"
+#include "IConcertEndpoint.h"
 #include "IConcertSession.h"
 
 class UConcertServerConfig;
@@ -41,6 +42,9 @@ public:
 	/** Gets the address of a remote admin endpoint, i.e. a client that is sending FConcertEndpointDiscoveryEvents. */
 	virtual FMessageAddress GetRemoteAddress(const FGuid& AdminEndpointId) const = 0;
 
+	/** Callback when a message has been acknowledged by a remote endpoint */
+	virtual FOnConcertMessageAcknowledgementReceivedFromLocalEndpoint& OnConcertMessageAcknowledgementReceived() = 0;
+
 	/**
 	 *	Returns if the server has already been started up.
 	 */
diff --git a/Engine/Plugins/Developer/Concert/ConcertMain/Source/ConcertTransport/Private/ConcertLocalEndpoint.cpp b/Engine/Plugins/Developer/Concert/ConcertMain/Source/ConcertTransport/Private/ConcertLocalEndpoint.cpp
--- a/Engine/Plugins/Developer/Concert/ConcertMain/Source/ConcertTransport/Private/ConcertLocalEndpoint.cpp
+++ b/Engine/Plugins/Developer/Concert/ConcertMain/Source/ConcertTransport/Private/ConcertLocalEndpoint.cpp
@@ -150,6 +150,11 @@ FOnConcertRemoteEndpointConnectionChanged& FConcertLocalEndpoint::OnRemoteEndpoi
 	return OnRemoteEndpointConnectionChangedDelegate;
 }
 
+FOnConcertMessageAcknowledgementReceivedFromLocalEndpoint& FConcertLocalEndpoint::OnConcertMessageAcknowledgementReceived()
+{
+	return OnConcertMessageAcknowledgementReceivedDelegate;
+}
+
 void FConcertLocalEndpoint::InternalAddRequestHandler(const FTopLevelAssetPath& RequestMessageType, const TSharedRef<IConcertRequestHandler>& Handler)
 {
 	RequestHandlers.Add(RequestMessageType, Handler);
@@ -279,6 +284,12 @@ FConcertRemoteEndpointRef FConcertLocalEndpoint::CreateRemoteEndpoint(const FCon
 		FScopeLock RemoteEndpointsLock(&RemoteEndpointsCS);
 		RemoteEndpoints.Add(InEndpointContext.EndpointId, NewRemoteEndpoint);
 	}
+	
+	NewRemoteEndpoint->OnConcertMessageAcknowledgementReceived().AddLambda(
+		[this](const FConcertEndpointContext& Context, const TSharedRef<IConcertMessage>& AckedMessage, const FConcertMessageContext& AckMessageContext)
+		{
+			OnConcertMessageAcknowledgementReceivedDelegate.Broadcast(GetEndpointContext(), Context, AckedMessage, AckMessageContext);
+		});
 	return NewRemoteEndpoint;
 }
 
diff --git a/Engine/Plugins/Developer/Concert/ConcertMain/Source/ConcertTransport/Private/ConcertLocalEndpoint.h b/Engine/Plugins/Developer/Concert/ConcertMain/Source/ConcertTransport/Private/ConcertLocalEndpoint.h
--- a/Engine/Plugins/Developer/Concert/ConcertMain/Source/ConcertTransport/Private/ConcertLocalEndpoint.h
+++ b/Engine/Plugins/Developer/Concert/ConcertMain/Source/ConcertTransport/Private/ConcertLocalEndpoint.h
@@ -34,6 +34,7 @@ public:
 	virtual FMessageAddress GetRemoteAddress(const FGuid& ConcertEndpointId) const override;
 
 	virtual FOnConcertRemoteEndpointConnectionChanged& OnRemoteEndpointConnectionChanged() override;
+	virtual FOnConcertMessageAcknowledgementReceivedFromLocalEndpoint& OnConcertMessageAcknowledgementReceived() override;
 
 protected:
 	virtual void InternalAddRequestHandler(const FTopLevelAssetPath& RequestMessageType, const TSharedRef<IConcertRequestHandler>& Handler) override;
@@ -158,6 +159,9 @@ private:
 	TArray<TTuple<FConcertEndpointContext, EConcertRemoteEndpointConnection>> PendingRemoteEndpointConnectionChangedEvents;
 	FOnConcertRemoteEndpointConnectionChanged OnRemoteEndpointConnectionChangedDelegate;
 
+	/** Callback when a message has been acknowledged by a remote endpoint */
+	FOnConcertMessageAcknowledgementReceivedFromLocalEndpoint OnConcertMessageAcknowledgementReceivedDelegate; 
+
 	/** Registered message handlers that do not generate a response */
 	TMap<FTopLevelAssetPath, TSharedPtr<IConcertEventHandler>> EventHandlers;
 
diff --git a/Engine/Plugins/Developer/Concert/ConcertMain/Source/ConcertTransport/Private/ConcertRemoteEndpoint.cpp b/Engine/Plugins/Developer/Concert/ConcertMain/Source/ConcertTransport/Private/ConcertRemoteEndpoint.cpp
--- a/Engine/Plugins/Developer/Concert/ConcertMain/Source/ConcertTransport/Private/ConcertRemoteEndpoint.cpp
+++ b/Engine/Plugins/Developer/Concert/ConcertMain/Source/ConcertTransport/Private/ConcertRemoteEndpoint.cpp
@@ -33,6 +33,11 @@ const FConcertEndpointContext& FConcertRemoteEndpoint::GetEndpointContext() cons
 	return EndpointContext;
 }
 
+FOnConcertMessageAcknowledgementReceived& FConcertRemoteEndpoint::OnConcertMessageAcknowledgementReceived()
+{
+	return OnConcertMessageAcknowledgementReceivedDelegate;
+}
+
 bool FConcertRemoteEndpoint::HasReliableChannel() const
 {
 	return ReliableChannelIdToReceive != FConcertMessageData::UnreliableChannelId;
@@ -197,6 +202,7 @@ void FConcertRemoteEndpoint::HandleAcknowledgement(const FConcertMessageContext&
 		if (PendingMessage->GetState() == EConcertMessageState::Pending)
 		{
 			PendingMessage->Acknowledge(Context);
+			OnConcertMessageAcknowledgementReceivedDelegate.Broadcast(GetEndpointContext(), PendingMessage, Context);
 		}
 		if (PendingMessage->GetState() == EConcertMessageState::Completed)
 		{
diff --git a/Engine/Plugins/Developer/Concert/ConcertMain/Source/ConcertTransport/Private/ConcertRemoteEndpoint.h b/Engine/Plugins/Developer/Concert/ConcertMain/Source/ConcertTransport/Private/ConcertRemoteEndpoint.h
--- a/Engine/Plugins/Developer/Concert/ConcertMain/Source/ConcertTransport/Private/ConcertRemoteEndpoint.h
+++ b/Engine/Plugins/Developer/Concert/ConcertMain/Source/ConcertTransport/Private/ConcertRemoteEndpoint.h
@@ -21,6 +21,8 @@ public:
 
 	virtual const FConcertEndpointContext& GetEndpointContext() const override;
 
+	virtual FOnConcertMessageAcknowledgementReceived& OnConcertMessageAcknowledgementReceived() override;
+
 	/** Does this remote endpoint have a reliable channel? */
 	bool HasReliableChannel() const;
 
@@ -143,6 +145,9 @@ private:
 	/** This context of this endpoint */
 	FConcertEndpointContext EndpointContext;
 
+	/** Callback when a message has been acknowledged by this remote endpoint */
+	FOnConcertMessageAcknowledgementReceived OnConcertMessageAcknowledgementReceivedDelegate;
+
 	/** The reliable channel ID to send with reliable messages */
 	uint16 ReliableChannelIdToSend;
 
diff --git a/Engine/Plugins/Developer/Concert/ConcertMain/Source/ConcertTransport/Public/IConcertEndpoint.h b/Engine/Plugins/Developer/Concert/ConcertMain/Source/ConcertTransport/Public/IConcertEndpoint.h
--- a/Engine/Plugins/Developer/Concert/ConcertMain/Source/ConcertTransport/Public/IConcertEndpoint.h
+++ b/Engine/Plugins/Developer/Concert/ConcertMain/Source/ConcertTransport/Public/IConcertEndpoint.h
@@ -4,9 +4,11 @@
 
 #include "CoreMinimal.h"
 #include "Async/Future.h"
+#include "ConcertTransportMessages.h"
 #include "IConcertMessageHandler.h"
 #include "IConcertMessages.h"
-#include "ConcertTransportMessages.h"
+
+struct FConcertMessageContext;
 
 /**
  * Context information for a Concert endpoint
@@ -39,6 +41,12 @@ enum class EConcertRemoteEndpointConnection : uint8
 	ClosedRemotely,
 };
 
+DECLARE_MULTICAST_DELEGATE_ThreeParams(FOnConcertMessageAcknowledgementReceived,
+	const FConcertEndpointContext& /*RemoteEndpoint*/,
+	const TSharedRef<IConcertMessage>& /*AckedMessage*/,
+	const FConcertMessageContext& /*AckMessageContext*/
+	);
+
 /** 
  * Interface representing a remote endpoint
  * that you can send to reliably or not from a local Concert endpoint
@@ -51,6 +59,9 @@ public:
 	
 	/** Get the context for this remote endpoint */
 	virtual const FConcertEndpointContext& GetEndpointContext() const = 0;
+	
+	/** Callback when a message has been acknowledged by this remote endpoint */
+	virtual FOnConcertMessageAcknowledgementReceived& OnConcertMessageAcknowledgementReceived() = 0;
 
 protected:
 	/** 
@@ -72,6 +83,12 @@ protected:
 
 
 DECLARE_MULTICAST_DELEGATE_TwoParams(FOnConcertRemoteEndpointConnectionChanged, const FConcertEndpointContext&, EConcertRemoteEndpointConnection);
+DECLARE_MULTICAST_DELEGATE_FourParams(FOnConcertMessageAcknowledgementReceivedFromLocalEndpoint,
+	const FConcertEndpointContext& /*LocalEndpoint*/,
+	const FConcertEndpointContext& /*RemoteEndpoint*/,
+	const TSharedRef<IConcertMessage>& /*AckedMessage*/,
+	const FConcertMessageContext& /*MessageContext*/
+);
 
 /**
  * Interface representing a local endpoint you can 
@@ -92,6 +109,9 @@ public:
 	/** Callback when a remote endpoint connection changes. */
 	virtual FOnConcertRemoteEndpointConnectionChanged& OnRemoteEndpointConnectionChanged() = 0;
 
+	/** Callback when a message has been acknowledged by a remote endpoint */
+	virtual FOnConcertMessageAcknowledgementReceivedFromLocalEndpoint& OnConcertMessageAcknowledgementReceived() = 0;
+
 	/**
 	 * Send a request to a remote endpoint
 	 * @param Request : the request to send, needs to be a UStruct that derives from FConcertRequestData
