Add way to get admin endpoints connected to server

diff --git a/Engine/Plugins/Developer/Concert/ConcertMain/Source/Concert/Private/ConcertServer.cpp b/Engine/Plugins/Developer/Concert/ConcertMain/Source/Concert/Private/ConcertServer.cpp
--- a/Engine/Plugins/Developer/Concert/ConcertMain/Source/Concert/Private/ConcertServer.cpp
+++ b/Engine/Plugins/Developer/Concert/ConcertMain/Source/Concert/Private/ConcertServer.cpp
@@ -150,6 +150,20 @@ const FConcertServerInfo& FConcertServer::GetServerInfo() const
 	return ServerInfo;
 }
 
+TArray<FConcertEndpointContext> FConcertServer::GetRemoteAdminEndpoints() const
+{
+	if (IsStarted())
+	{
+		return ServerAdminEndpoint->GetRemoteEndpoints();
+	}
+	return {};
+}
+
+FOnConcertRemoteEndpointConnectionChanged& FConcertServer::OnRemoteEndpointConnectionChanged()
+{
+	return OnConcertRemoteEndpointConnectionChangedDelegate;
+}
+
 FMessageAddress FConcertServer::GetRemoteAddress(const FGuid& AdminEndpointId) const
 {
 	if (IsStarted())
@@ -188,6 +202,10 @@ void FConcertServer::Startup()
 			{
 				OnConcertMessageAcknowledgementReceivedFromLocalEndpoint.Broadcast(LocalEndpoint, RemoteEndpoint, AckedMessage, MessageContext);
 			});
+		ServerAdminEndpoint->OnRemoteEndpointConnectionChanged().AddLambda([this](const FConcertEndpointContext& Context, EConcertRemoteEndpointConnection Connection)
+		{
+			OnConcertRemoteEndpointConnectionChangedDelegate.Broadcast(Context, Connection);
+		});
 
 		// Make it discoverable
 		ServerAdminEndpoint->SubscribeEventHandler<FConcertAdmin_DiscoverServersEvent>(this, &FConcertServer::HandleDiscoverServersEvent);
diff --git a/Engine/Plugins/Developer/Concert/ConcertMain/Source/Concert/Private/ConcertServer.h b/Engine/Plugins/Developer/Concert/ConcertMain/Source/Concert/Private/ConcertServer.h
--- a/Engine/Plugins/Developer/Concert/ConcertMain/Source/Concert/Private/ConcertServer.h
+++ b/Engine/Plugins/Developer/Concert/ConcertMain/Source/Concert/Private/ConcertServer.h
@@ -25,6 +25,9 @@ public:
 	virtual bool IsConfigured() const override;
 	virtual const UConcertServerConfig* GetConfiguration() const override;
 	virtual const FConcertServerInfo& GetServerInfo() const override;
+
+	virtual TArray<FConcertEndpointContext> GetRemoteAdminEndpoints() const override;
+	virtual FOnConcertRemoteEndpointConnectionChanged& OnRemoteEndpointConnectionChanged() override;
 	virtual FMessageAddress GetRemoteAddress(const FGuid& AdminEndpointId) const override;
 	virtual FOnConcertMessageAcknowledgementReceivedFromLocalEndpoint& OnConcertMessageAcknowledgementReceived() override;
 
@@ -179,6 +182,9 @@ private:
 	
 	/** Administration endpoint for the server (i.e. creating, joining sessions) */
 	TSharedPtr<IConcertLocalEndpoint> ServerAdminEndpoint;
+
+	/** Called when ServerAdminEndpoint emits IConcertLocalEndpoint::OnConcertRemoteEndpointConnectionChanged */
+	FOnConcertRemoteEndpointConnectionChanged OnConcertRemoteEndpointConnectionChangedDelegate; 
 	
 	/** Server and Instance Info */
 	FConcertServerInfo ServerInfo;
diff --git a/Engine/Plugins/Developer/Concert/ConcertMain/Source/Concert/Public/IConcertServer.h b/Engine/Plugins/Developer/Concert/ConcertMain/Source/Concert/Public/IConcertServer.h
--- a/Engine/Plugins/Developer/Concert/ConcertMain/Source/Concert/Public/IConcertServer.h
+++ b/Engine/Plugins/Developer/Concert/ConcertMain/Source/Concert/Public/IConcertServer.h
@@ -39,7 +39,19 @@ public:
 	 */
 	virtual const FConcertServerInfo& GetServerInfo() const = 0;
 
-	/** Gets the address of a remote admin endpoint, i.e. a client that is sending FConcertEndpointDiscoveryEvents. */
+	
+	/**
+	 * Gets all remote admin endpoint IDs.
+	 * A remote admin endpoint is not connected to any session and communicating with the server otherwise, e.g. discovering sessions.
+	 */
+	virtual TArray<FConcertEndpointContext> GetRemoteAdminEndpoints() const = 0;
+	
+	/** Callback when a remote admin endpoint connection changes. */
+	virtual FOnConcertRemoteEndpointConnectionChanged& OnRemoteEndpointConnectionChanged() = 0;
+	
+	/**
+	 * Gets the address of a remote admin endpoint, i.e. a client that is sending FConcertEndpointDiscoveryEvents. 
+	 */
 	virtual FMessageAddress GetRemoteAddress(const FGuid& AdminEndpointId) const = 0;
 
 	/** Callback when a message has been acknowledged by a remote endpoint */
diff --git a/Engine/Plugins/Developer/Concert/ConcertMain/Source/ConcertTransport/Private/ConcertLocalEndpoint.cpp b/Engine/Plugins/Developer/Concert/ConcertMain/Source/ConcertTransport/Private/ConcertLocalEndpoint.cpp
--- a/Engine/Plugins/Developer/Concert/ConcertMain/Source/ConcertTransport/Private/ConcertLocalEndpoint.cpp
+++ b/Engine/Plugins/Developer/Concert/ConcertMain/Source/ConcertTransport/Private/ConcertLocalEndpoint.cpp
@@ -6,6 +6,7 @@
 
 #include "MessageEndpoint.h"
 #include "MessageEndpointBuilder.h"
+#include "Algo/Transform.h"
 
 #include "HAL/Runnable.h"
 #include "HAL/RunnableThread.h"
@@ -137,6 +138,17 @@ const FConcertEndpointContext& FConcertLocalEndpoint::GetEndpointContext() const
 	return EndpointContext;
 }
 
+TArray<FConcertEndpointContext> FConcertLocalEndpoint::GetRemoteEndpoints() const
+{
+	FScopeLock RemoteEndpointsLock(&RemoteEndpointsCS);
+	TArray<FConcertEndpointContext> Result;
+	Algo::Transform(RemoteEndpoints, Result, [](const TPair<FGuid, FConcertRemoteEndpointPtr>& RemoteEndpoint)
+	{
+		return RemoteEndpoint.Value->GetEndpointContext();
+	});
+	return Result;
+}
+
 FMessageAddress FConcertLocalEndpoint::GetRemoteAddress(const FGuid& ConcertEndpointId) const
 {
 	const FConcertRemoteEndpointPtr RemoteEndpoint = FindRemoteEndpoint(ConcertEndpointId);
diff --git a/Engine/Plugins/Developer/Concert/ConcertMain/Source/ConcertTransport/Private/ConcertLocalEndpoint.h b/Engine/Plugins/Developer/Concert/ConcertMain/Source/ConcertTransport/Private/ConcertLocalEndpoint.h
--- a/Engine/Plugins/Developer/Concert/ConcertMain/Source/ConcertTransport/Private/ConcertLocalEndpoint.h
+++ b/Engine/Plugins/Developer/Concert/ConcertMain/Source/ConcertTransport/Private/ConcertLocalEndpoint.h
@@ -31,6 +31,7 @@ public:
 	virtual ~FConcertLocalEndpoint();
 
 	virtual const FConcertEndpointContext& GetEndpointContext() const override;
+	virtual TArray<FConcertEndpointContext> GetRemoteEndpoints() const override;
 	virtual FMessageAddress GetRemoteAddress(const FGuid& ConcertEndpointId) const override;
 
 	virtual FOnConcertRemoteEndpointConnectionChanged& OnRemoteEndpointConnectionChanged() override;
diff --git a/Engine/Plugins/Developer/Concert/ConcertMain/Source/ConcertTransport/Public/IConcertEndpoint.h b/Engine/Plugins/Developer/Concert/ConcertMain/Source/ConcertTransport/Public/IConcertEndpoint.h
--- a/Engine/Plugins/Developer/Concert/ConcertMain/Source/ConcertTransport/Public/IConcertEndpoint.h
+++ b/Engine/Plugins/Developer/Concert/ConcertMain/Source/ConcertTransport/Public/IConcertEndpoint.h
@@ -103,6 +103,9 @@ public:
 	/** Get the context for this endpoint. */
 	virtual const FConcertEndpointContext& GetEndpointContext() const = 0;
 
+	/** Gets all remote endpoints connected to this local endpoint. */
+	virtual TArray<FConcertEndpointContext> GetRemoteEndpoints() const = 0;
+	
 	/** Gets the remote address of a client or the server. If ConcertEndpointId is not valid, the result's IsValid return false. */
 	virtual FMessageAddress GetRemoteAddress(const FGuid& ConcertEndpointId) const = 0;
 
