Added more LLM tags to VP code.
Renamed some of the tags to use the correct plugin name for the parent.

diff --git a/Engine/Plugins/Animation/LiveLink/Source/LiveLink/Private/LiveLinkClient.cpp b/Engine/Plugins/Animation/LiveLink/Source/LiveLink/Private/LiveLinkClient.cpp
--- a/Engine/Plugins/Animation/LiveLink/Source/LiveLink/Private/LiveLinkClient.cpp
+++ b/Engine/Plugins/Animation/LiveLink/Source/LiveLink/Private/LiveLinkClient.cpp
@@ -40,6 +40,7 @@
 #include "VirtualSubjects/LiveLinkBlueprintVirtualSubject.h"
 #endif
 
+LLM_DEFINE_TAG(LiveLink_LiveLinkClient);
 
 /**
  * Declare stats to see what takes up time in LiveLink
@@ -66,6 +67,8 @@ static TAutoConsoleVariable<int32> CVarMaxNewFrameDataPerUpdate(
 
 FLiveLinkClient::FLiveLinkClient()
 {
+	LLM_SCOPE_BYTAG(LiveLink_LiveLinkClient);
+
 #if WITH_EDITOR
 	CachedEngineTime = 0.0;
 #endif
@@ -201,6 +204,8 @@ void FLiveLinkClient::CacheValues()
 
 void FLiveLinkClient::Shutdown()
 {
+	LLM_SCOPE_BYTAG(LiveLink_LiveLinkClient);
+
 	/*if(IMediaModule* MediaModule = FModuleManager::GetModulePtr<IMediaModule>("Media"))
 	{
 		MediaModule->GetOnTickPreEngineCompleted().RemoveAll(this);
diff --git a/Engine/Plugins/Animation/LiveLink/Source/LiveLink/Private/LiveLinkLogInstance.cpp b/Engine/Plugins/Animation/LiveLink/Source/LiveLink/Private/LiveLinkLogInstance.cpp
--- a/Engine/Plugins/Animation/LiveLink/Source/LiveLink/Private/LiveLinkLogInstance.cpp
+++ b/Engine/Plugins/Animation/LiveLink/Source/LiveLink/Private/LiveLinkLogInstance.cpp
@@ -15,7 +15,7 @@
 #include "Modules/ModuleManager.h"
 #endif
 
-
+LLM_DEFINE_TAG(LiveLink_LiveLinkLogInstance);
 #define LOCTEXT_NAMESPACE "LiveLinkLogInstance"
 
 
@@ -90,6 +90,7 @@ namespace LiveLinkLogDetail
 
 void FLiveLinkLogInstance::CreateInstance()
 {
+	LLM_SCOPE_BYTAG(LiveLink_LiveLinkLogInstance);
 	ensure(FLiveLinkLog::Instance == nullptr);
 	if (FLiveLinkLog::Instance)
 	{
@@ -102,6 +103,7 @@ void FLiveLinkLogInstance::CreateInstance()
 
 void FLiveLinkLogInstance::DestroyInstance()
 {
+	LLM_SCOPE_BYTAG(LiveLink_LiveLinkLogInstance);
 	ensure(FLiveLinkLog::Instance);
 	FLiveLinkLog::Instance.Reset();
 	FLiveLinkLog::Instance = nullptr;
diff --git a/Engine/Plugins/Animation/LiveLink/Source/LiveLink/Private/LiveLinkMessageBusDiscoveryManager.cpp b/Engine/Plugins/Animation/LiveLink/Source/LiveLink/Private/LiveLinkMessageBusDiscoveryManager.cpp
--- a/Engine/Plugins/Animation/LiveLink/Source/LiveLink/Private/LiveLinkMessageBusDiscoveryManager.cpp
+++ b/Engine/Plugins/Animation/LiveLink/Source/LiveLink/Private/LiveLinkMessageBusDiscoveryManager.cpp
@@ -10,6 +10,7 @@
 #include "MessageEndpointBuilder.h"
 #include "Misc/ScopeLock.h"
 
+LLM_DEFINE_TAG(LiveLink_LiveLinkMessageBusDiscoveryManager);
 
 #define LL_HEARTBEAT_SLEEP_TIME 1.0f
 
@@ -17,6 +18,8 @@ FLiveLinkMessageBusDiscoveryManager::FLiveLinkMessageBusDiscoveryManager()
 	: bRunning(true)
 	, Thread(nullptr)
 {
+	LLM_SCOPE_BYTAG(LiveLink_LiveLinkMessageBusDiscoveryManager);
+
 	PingRequestCounter = 0;
 	PingRequestFrequency = GetDefault<ULiveLinkSettings>()->GetMessageBusPingRequestFrequency();
 
diff --git a/Engine/Plugins/Animation/LiveLink/Source/LiveLink/Private/LiveLinkModule.cpp b/Engine/Plugins/Animation/LiveLink/Source/LiveLink/Private/LiveLinkModule.cpp
--- a/Engine/Plugins/Animation/LiveLink/Source/LiveLink/Private/LiveLinkModule.cpp
+++ b/Engine/Plugins/Animation/LiveLink/Source/LiveLink/Private/LiveLinkModule.cpp
@@ -10,6 +10,7 @@
 #include "Misc/CoreDelegates.h"
 #include "Styling/SlateStyleRegistry.h"
 
+LLM_DEFINE_TAG(LiveLink);
 #define LOCTEXT_NAMESPACE "LiveLinkModule"
 
 FLiveLinkClient* FLiveLinkModule::LiveLinkClient_AnyThread = nullptr;
@@ -27,6 +28,7 @@ FLiveLinkModule::FLiveLinkModule()
 
 void FLiveLinkModule::StartupModule()
 {
+	LLM_SCOPE_BYTAG(LiveLink);
 	FLiveLinkLogInstance::CreateInstance();
 	CreateStyle();
 
@@ -40,6 +42,7 @@ void FLiveLinkModule::StartupModule()
 
 void FLiveLinkModule::ShutdownModule()
 {
+	LLM_SCOPE_BYTAG(LiveLink);
 	FCoreDelegates::OnFEngineLoopInitComplete.RemoveAll(this);
 
 	HeartbeatEmitter->Exit();
diff --git a/Engine/Plugins/Animation/LiveLink/Source/LiveLinkEditor/Private/LiveLinkEditorModule.cpp b/Engine/Plugins/Animation/LiveLink/Source/LiveLinkEditor/Private/LiveLinkEditorModule.cpp
--- a/Engine/Plugins/Animation/LiveLink/Source/LiveLinkEditor/Private/LiveLinkEditorModule.cpp
+++ b/Engine/Plugins/Animation/LiveLink/Source/LiveLinkEditor/Private/LiveLinkEditorModule.cpp
@@ -46,7 +46,7 @@
  * Implements the LiveLinkEditor module.
  */
 
-
+LLM_DEFINE_TAG(LiveLink_LiveLinkEditor);
 DEFINE_LOG_CATEGORY(LogLiveLinkEditor);
 
 
@@ -79,6 +79,8 @@ public:
 
 	virtual void StartupModule() override
 	{
+		LLM_SCOPE_BYTAG(LiveLink_LiveLinkEditor);
+
 		static FName LiveLinkStyle(TEXT("LiveLinkStyle"));
 		StyleSet = MakeShared<FSlateStyleSet>(LiveLinkStyle);
 
@@ -160,6 +162,8 @@ public:
 
 	virtual void ShutdownModule() override
 	{
+		LLM_SCOPE_BYTAG(LiveLink_LiveLinkEditor);
+
 		UnregisterCustomizations();
 
 		UnregisterTabSpawner();
diff --git a/Engine/Plugins/Animation/LiveLink/Source/LiveLinkSequencer/Private/LiveLinkSequencerModule.cpp b/Engine/Plugins/Animation/LiveLink/Source/LiveLinkSequencer/Private/LiveLinkSequencerModule.cpp
--- a/Engine/Plugins/Animation/LiveLink/Source/LiveLinkSequencer/Private/LiveLinkSequencerModule.cpp
+++ b/Engine/Plugins/Animation/LiveLink/Source/LiveLinkSequencer/Private/LiveLinkSequencerModule.cpp
@@ -24,7 +24,7 @@
 #include "Templates/SubclassOf.h"
 #include "Widgets/Docking/SDockTab.h"
 
-
+LLM_DEFINE_TAG(LiveLink_LiveLinkSequencer);
 DEFINE_LOG_CATEGORY(LogLiveLinkSequencer);
 
 
@@ -63,6 +63,8 @@ public:
 
 	virtual void StartupModule() override
 	{
+		LLM_SCOPE_BYTAG(LiveLink_LiveLinkSequencer);
+
 		static FName LiveLinkSequencerStyle(TEXT("LiveLinkSequencerStyle"));
 		StyleSet = MakeShared<FSlateStyleSet>(LiveLinkSequencerStyle);
 
@@ -100,6 +102,8 @@ public:
 
 	virtual void ShutdownModule() override
 	{
+		LLM_SCOPE_BYTAG(LiveLink_LiveLinkSequencer);
+
 		UnregisterTakeRecorderSourceMenuExtender();
 
 		FModuleManager::Get().OnModulesChanged().Remove(ModulesChangedHandle);
diff --git a/Engine/Plugins/Compositing/Composure/Source/Composure/Private/CompositingElements/CompositingElementTransforms.cpp b/Engine/Plugins/Compositing/Composure/Source/Composure/Private/CompositingElements/CompositingElementTransforms.cpp
--- a/Engine/Plugins/Compositing/Composure/Source/Composure/Private/CompositingElements/CompositingElementTransforms.cpp
+++ b/Engine/Plugins/Compositing/Composure/Source/Composure/Private/CompositingElements/CompositingElementTransforms.cpp
@@ -303,6 +303,8 @@ UTexture* UMultiPassChromaKeyer::ApplyTransform_Implementation(UTexture* Input,
 #include "UObject/ConstructorHelpers.h"
 #include "Materials/MaterialInterface.h"
 
+LLM_DEFINE_TAG(Composure_MultiPassDespill);
+
 namespace MultiPassDespill_Impl
 {
 	static const FName ColorPlateKeyName(TEXT("ColorPlateParamName"));
@@ -312,6 +314,8 @@ namespace MultiPassDespill_Impl
 
 UMultiPassDespill::UMultiPassDespill()
 {
+	LLM_SCOPE_BYTAG(Composure_MultiPassDespill);
+
 	static ConstructorHelpers::FObjectFinder<UMaterialInterface> GKeyerMaterial(TEXT("/Composure/Materials/ChromaKeying/M_SinglePassDespill"));
 	KeyerMaterial.Material = Cast<UMaterialInterface>(GKeyerMaterial.Object);
 
@@ -331,6 +335,8 @@ UMultiPassDespill::UMultiPassDespill()
 
 UTexture* UMultiPassDespill::ApplyTransform_Implementation(UTexture* Input, UComposurePostProcessingPassProxy* PostProcessProxy, ACameraActor* /*TargetCamera*/)
 {
+	LLM_SCOPE_BYTAG(Composure_MultiPassDespill);
+
 	FCompositingTargetSwapChain SwapChain = FCompositingTargetSwapChain(SharedTargetPool, /*RenderScale =*/1.f);
 
 	UTexture* Result = Input;
diff --git a/Engine/Plugins/Compositing/Composure/Source/Composure/Private/MovieScene/MovieSceneComposurePostMoveSettingsSection.cpp b/Engine/Plugins/Compositing/Composure/Source/Composure/Private/MovieScene/MovieSceneComposurePostMoveSettingsSection.cpp
--- a/Engine/Plugins/Compositing/Composure/Source/Composure/Private/MovieScene/MovieSceneComposurePostMoveSettingsSection.cpp
+++ b/Engine/Plugins/Compositing/Composure/Source/Composure/Private/MovieScene/MovieSceneComposurePostMoveSettingsSection.cpp
@@ -5,6 +5,8 @@
 #include "Channels/MovieSceneChannelProxy.h"
 #include "ComposurePostMoves.h"
 
+LLM_DEFINE_TAG(Composure_MovieSceneComposurePostMoveSettingsSection);
+
 #if WITH_EDITOR
 
 struct FPostMoveSettingsChannelEditorData
@@ -87,7 +89,9 @@ struct FPostMoveSettingsChannelEditorData
 
 UMovieSceneComposurePostMoveSettingsSection::UMovieSceneComposurePostMoveSettingsSection( const FObjectInitializer& ObjectInitializer )
 	: Super( ObjectInitializer )
-{ 
+{
+	LLM_SCOPE_BYTAG(Composure_MovieSceneComposurePostMoveSettingsSection);
+
 	EvalOptions.EnableAndSetCompletionMode
 		(GetLinkerCustomVersion(FSequencerObjectVersion::GUID) < FSequencerObjectVersion::WhenFinishedDefaultsToProjectDefault ? 
 			EMovieSceneCompletionMode::RestoreState : 
diff --git a/Engine/Plugins/Compositing/Composure/Source/ComposureEditor/Private/ComposureEditorModule.cpp b/Engine/Plugins/Compositing/Composure/Source/ComposureEditor/Private/ComposureEditorModule.cpp
--- a/Engine/Plugins/Compositing/Composure/Source/ComposureEditor/Private/ComposureEditorModule.cpp
+++ b/Engine/Plugins/Compositing/Composure/Source/ComposureEditor/Private/ComposureEditorModule.cpp
@@ -6,14 +6,15 @@
 #include "Sequencer/ComposurePostMoveSettingsPropertyTrackEditor.h"
 #include "Sequencer/ComposureExportTrackEditor.h"
 
-LLM_DEFINE_TAG(Compositing_ComposureEditor);
+LLM_DEFINE_TAG(Composure_ComposureEditor);
 DEFINE_LOG_CATEGORY(LogComposureEditor);
 
 class FComposureEditorModule : public IModuleInterface
 {
 	virtual void StartupModule() override
 	{
-		LLM_SCOPE_BYTAG(Compositing_ComposureEditor);
+		LLM_SCOPE_BYTAG(Composure_ComposureEditor);
+
 		ISequencerModule& SequencerModule = FModuleManager::LoadModuleChecked<ISequencerModule>("Sequencer");
 		CreatePostMoveSettingsPropertyTrackEditorHandle = SequencerModule.RegisterPropertyTrackEditor<FComposurePostMoveSettingsPropertyTrackEditor>();
 		ComposureExportTrackEditorHandle               = SequencerModule.RegisterTrackEditor(FOnCreateTrackEditor::CreateLambda([](TSharedRef<ISequencer> In){ return MakeShared<FComposureExportTrackEditor>(In); }));
@@ -21,6 +22,8 @@ class FComposureEditorModule : public IModuleInterface
 
 	virtual void ShutdownModule() override
 	{
+		LLM_SCOPE_BYTAG(Composure_ComposureEditor);
+
 		ISequencerModule* SequencerModule = FModuleManager::GetModulePtr<ISequencerModule>("Sequencer");
 		if (SequencerModule != nullptr)
 		{
diff --git a/Engine/Plugins/Compositing/Composure/Source/ComposureLayersEditor/Private/CompElementEditorModule.cpp b/Engine/Plugins/Compositing/Composure/Source/ComposureLayersEditor/Private/CompElementEditorModule.cpp
--- a/Engine/Plugins/Compositing/Composure/Source/ComposureLayersEditor/Private/CompElementEditorModule.cpp
+++ b/Engine/Plugins/Compositing/Composure/Source/ComposureLayersEditor/Private/CompElementEditorModule.cpp
@@ -25,7 +25,7 @@
 #include "LevelEditorViewport.h"
 #include "Widgets/SCompElementPreviewDialog.h"
 
-LLM_DEFINE_TAG(Compositing_ComposureLayersEditor);
+LLM_DEFINE_TAG(Composure_ComposureLayersEditor);
 
 namespace CompElementEditor_Impl
 {
@@ -92,7 +92,8 @@ private:
 //------------------------------------------------------------------------------
 void FCompElementEditorModule::StartupModule() 
 {
-	LLM_SCOPE_BYTAG(Compositing_ComposureLayersEditor);
+	LLM_SCOPE_BYTAG(Composure_ComposureLayersEditor);
+
 	FComposureEditorStyle::Get();
 	FCompElementEditorCommands::Register();
 
@@ -121,6 +122,8 @@ void FCompElementEditorModule::StartupModule()
 //------------------------------------------------------------------------------
 void FCompElementEditorModule::ShutdownModule()
 {
+	LLM_SCOPE_BYTAG(Composure_ComposureLayersEditor);
+
 	IModularFeatures::Get().UnregisterModularFeature(ICompositingEditor::GetModularFeatureName(), this);
 
 	if (FSlateApplication::IsInitialized())
@@ -275,7 +278,7 @@ void FCompElementEditorModule::RegisterEditorTab()
 
 	LevelEditorTabManagerChangedHandle = LevelEditorModule.OnTabManagerChanged().AddLambda([]()
 	{
-		LLM_SCOPE_BYTAG(Compositing_ComposureLayersEditor);
+		LLM_SCOPE_BYTAG(Composure_ComposureLayersEditor);
 		FLevelEditorModule& LevelEditorModule = FModuleManager::GetModuleChecked<FLevelEditorModule>(CompElementEditor_Impl::LevelEditorModuleName);
 		TSharedPtr<FTabManager> LevelEditorTabManager = LevelEditorModule.GetLevelEditorTabManager();
 
diff --git a/Engine/Plugins/Compositing/OpenColorIO/Source/OpenColorIOEditor/Private/OpenColorIOEditorModule.cpp b/Engine/Plugins/Compositing/OpenColorIO/Source/OpenColorIOEditor/Private/OpenColorIOEditorModule.cpp
--- a/Engine/Plugins/Compositing/OpenColorIO/Source/OpenColorIOEditor/Private/OpenColorIOEditorModule.cpp
+++ b/Engine/Plugins/Compositing/OpenColorIO/Source/OpenColorIOEditor/Private/OpenColorIOEditorModule.cpp
@@ -23,7 +23,7 @@
 #include "ToolMenus.h"
 #include "Widgets/SOpenColorIODisplay.h"
 
-LLM_DEFINE_TAG(Compositing_OpenColorIOEditor);
+LLM_DEFINE_TAG(OpenColorIO_OpenColorIOEditor);
 DEFINE_LOG_CATEGORY(LogOpenColorIOEditor);
 
 #define LOCTEXT_NAMESPACE "OpenColorIOEditorModule"
@@ -31,7 +31,8 @@ DEFINE_LOG_CATEGORY(LogOpenColorIOEditor);
 
 void FOpenColorIOEditorModule::StartupModule()
 {
-	LLM_SCOPE_BYTAG(Compositing_OpenColorIOEditor);
+	LLM_SCOPE_BYTAG(OpenColorIO_OpenColorIOEditor);
+
 	FWorldDelegates::OnPreWorldInitialization.AddRaw(this, &FOpenColorIOEditorModule::OnWorldInit);
 
 	// Register asset type actions for OpenColorIOConfiguration class
@@ -49,6 +50,8 @@ void FOpenColorIOEditorModule::StartupModule()
 
 void FOpenColorIOEditorModule::ShutdownModule()
 {
+	LLM_SCOPE_BYTAG(OpenColorIO_OpenColorIOEditor);
+
 	UnregisterViewMenuExtension();
 	UnregisterStyle();
 	UnregisterCustomizations();
@@ -92,6 +95,7 @@ void FOpenColorIOEditorModule::UnregisterCustomizations()
 
 void FOpenColorIOEditorModule::OnWorldInit(UWorld* InWorld, const UWorld::InitializationValues InInitializationValues)
 {
+	LLM_SCOPE_BYTAG(OpenColorIO_OpenColorIOEditor);
 	if (InWorld && InWorld->WorldType == EWorldType::Editor)
 	{
 		CleanFeatureLevelDelegate();
@@ -252,6 +256,7 @@ void FOpenColorIOEditorModule::OnLevelViewportClientListChanged()
 
 void FOpenColorIOEditorModule::OnEngineLoopInitComplete()
 {
+	LLM_SCOPE_BYTAG(OpenColorIO_OpenColorIOEditor);
 	// Register for Viewport updates to be able to track them
 	GEditor->OnLevelViewportClientListChanged().AddRaw(this, &FOpenColorIOEditorModule::OnLevelViewportClientListChanged);
 }
diff --git a/Engine/Plugins/Developer/Concert/ConcertApp/DisasterRecoveryClient/Source/DisasterRecoveryClient/Private/DisasterRecoveryClientModule.cpp b/Engine/Plugins/Developer/Concert/ConcertApp/DisasterRecoveryClient/Source/DisasterRecoveryClient/Private/DisasterRecoveryClientModule.cpp
--- a/Engine/Plugins/Developer/Concert/ConcertApp/DisasterRecoveryClient/Source/DisasterRecoveryClient/Private/DisasterRecoveryClientModule.cpp
+++ b/Engine/Plugins/Developer/Concert/ConcertApp/DisasterRecoveryClient/Source/DisasterRecoveryClient/Private/DisasterRecoveryClientModule.cpp
@@ -132,6 +132,8 @@ public:
 
 	virtual void ShutdownModule() override
 	{
+		LLM_SCOPE_BYTAG(Concert_DisasterRecoveryClient);
+
 		FCoreDelegates::OnFEngineLoopInitComplete.RemoveAll(this);
 
 		// Unhook AppPreExit and call it
@@ -349,6 +351,8 @@ private:
 	/** Starts or restarts the recovery service. */
 	void StartDisasterRecoveryService()
 	{
+		LLM_SCOPE_BYTAG(Concert_DisasterRecoveryClient);
+
 		// Started for the first time (not restarted)? Always create the manager (and start the service) even if disaster recovery is disabled (in the settings). This allows the user to import sessions for crash analysis.
 		if (!SessionManager)
 		{
diff --git a/Engine/Plugins/Developer/Concert/ConcertApp/MultiUserClient/Source/MultiUserClient/Private/ConcertWorkspaceUI.cpp b/Engine/Plugins/Developer/Concert/ConcertApp/MultiUserClient/Source/MultiUserClient/Private/ConcertWorkspaceUI.cpp
--- a/Engine/Plugins/Developer/Concert/ConcertApp/MultiUserClient/Source/MultiUserClient/Private/ConcertWorkspaceUI.cpp
+++ b/Engine/Plugins/Developer/Concert/ConcertApp/MultiUserClient/Source/MultiUserClient/Private/ConcertWorkspaceUI.cpp
@@ -45,6 +45,7 @@
 #include "Widgets/Docking/SDockTab.h"
 #include "Widgets/Notifications/SNotificationList.h"
 
+LLM_DEFINE_TAG(Concert_ConcertWorkspaceUI);
 #define LOCTEXT_NAMESPACE "ConcertFrontend"
 
 static const FName ConcertHistoryTabName(TEXT("ConcertHistory"));
@@ -1005,24 +1006,29 @@ void FConcertWorkspaceUI::GenerateConcertAssetContextMenu(FMenuBuilder& MenuBuil
 }
 
 TSharedRef<SWidget> FConcertWorkspaceUI::OnGenerateAssetViewLockStateIcons(const FAssetData& AssetData)
-{	return SNew(SConcertWorkspaceLockStateIndicator, AsShared())
+{	
+	LLM_SCOPE_BYTAG(Concert_ConcertWorkspaceUI);
+	return SNew(SConcertWorkspaceLockStateIndicator, AsShared())
 		.AssetPath(AssetData.PackageName);
 }
 
 TSharedRef<SWidget> FConcertWorkspaceUI::OnGenerateAssetViewLockStateTooltip(const FAssetData& AssetData)
 {
+	LLM_SCOPE_BYTAG(Concert_ConcertWorkspaceUI);
 	return SNew(SConcertWorkspaceLockStateTooltip, AsShared())
 		.AssetPath(AssetData.PackageName);
 }
 
 TSharedRef<SWidget> FConcertWorkspaceUI::OnGenerateAssetViewModifiedByOtherIcon(const FAssetData& AssetData)
 {
+	LLM_SCOPE_BYTAG(Concert_ConcertWorkspaceUI);
 	return SNew(SConcertWorkspaceModifiedByOtherIndicator, AsShared())
 		.AssetPath(AssetData.PackageName);
 }
 
 TSharedRef<SWidget> FConcertWorkspaceUI::OnGenerateAssetViewModifiedByOtherTooltip(const FAssetData& AssetData)
 {
+	LLM_SCOPE_BYTAG(Concert_ConcertWorkspaceUI);
 	return SNew(SConcertWorkspaceModifiedByOtherTooltip, AsShared())
 		.AssetPath(AssetData.PackageName);
 }
diff --git a/Engine/Plugins/Developer/Concert/ConcertMain/Source/Concert/Private/ConcertClient.cpp b/Engine/Plugins/Developer/Concert/ConcertMain/Source/Concert/Private/ConcertClient.cpp
--- a/Engine/Plugins/Developer/Concert/ConcertMain/Source/Concert/Private/ConcertClient.cpp
+++ b/Engine/Plugins/Developer/Concert/ConcertMain/Source/Concert/Private/ConcertClient.cpp
@@ -674,6 +674,7 @@ const FString& FConcertClient::GetRole() const
 void FConcertClient::Configure(const UConcertClientConfig* InSettings)
 {
 	LLM_SCOPE_BYTAG(Concert_ConcertClient);
+
 	ClientInfo.Initialize();
 	check(InSettings != nullptr);
 	Settings = TStrongObjectPtr<const UConcertClientConfig>(InSettings);
@@ -881,6 +882,8 @@ EConcertConnectionStatus FConcertClient::GetSessionConnectionStatus() const
 
 TFuture<EConcertResponseCode> FConcertClient::CreateSession(const FGuid& ServerAdminEndpointId, const FConcertCreateSessionArgs& CreateSessionArgs)
 {
+	LLM_SCOPE_BYTAG(Concert_ConcertClient);
+
 	// We don't want the client to get automatically reconnected to it's default session if something wrong happens
 	AutoConnection.Reset();
 	return InternalCreateSession(ServerAdminEndpointId, CreateSessionArgs);
@@ -888,6 +891,8 @@ TFuture<EConcertResponseCode> FConcertClient::CreateSession(const FGuid& ServerA
 
 TFuture<EConcertResponseCode> FConcertClient::JoinSession(const FGuid& ServerAdminEndpointId, const FGuid& SessionId)
 {
+	LLM_SCOPE_BYTAG(Concert_ConcertClient);
+
 	// We don't want the client to get automatically reconnected to it's default session if something wrong happens
 	AutoConnection.Reset();
 	return InternalJoinSession(ServerAdminEndpointId, SessionId);
@@ -895,6 +900,8 @@ TFuture<EConcertResponseCode> FConcertClient::JoinSession(const FGuid& ServerAdm
 
 TFuture<EConcertResponseCode> FConcertClient::RestoreSession(const FGuid& ServerAdminEndpointId, const FConcertCopySessionArgs& RestoreSessionArgs)
 {
+	LLM_SCOPE_BYTAG(Concert_ConcertClient);
+
 	// We don't want the client to get automatically reconnected to the default session if something wrong happens
 	if (RestoreSessionArgs.bAutoConnect)
 	{
@@ -905,6 +912,8 @@ TFuture<EConcertResponseCode> FConcertClient::RestoreSession(const FGuid& Server
 
 TFuture<EConcertResponseCode> FConcertClient::CopySession(const FGuid& ServerAdminEndpointId, const FConcertCopySessionArgs& CopySessionArgs)
 {
+	LLM_SCOPE_BYTAG(Concert_ConcertClient);
+	
 	// We don't want the client to get automatically reconnected to the default session if the copy/connect fails.
 	if (CopySessionArgs.bAutoConnect)
 	{
@@ -1101,6 +1110,8 @@ TSharedPtr<IConcertClientSession> FConcertClient::GetCurrentSession() const
 
 TFuture<FConcertAdmin_MountSessionRepositoryResponse> FConcertClient::MountSessionRepository(const FGuid& ServerAdminEndpointId, const FString& RepositoryRootDir, const FGuid& RepositoryId, bool bCreateIfNotExist, bool bAsDefault) const
 {
+	LLM_SCOPE_BYTAG(Concert_ConcertClient);
+
 	FConcertAdmin_MountSessionRepositoryRequest MountRepositoryRequest;
 	MountRepositoryRequest.RepositoryId = RepositoryId;
 	MountRepositoryRequest.RepositoryRootDir = RepositoryRootDir;
@@ -1111,12 +1122,16 @@ TFuture<FConcertAdmin_MountSessionRepositoryResponse> FConcertClient::MountSessi
 
 TFuture<FConcertAdmin_GetSessionRepositoriesResponse> FConcertClient::GetSessionRepositories(const FGuid& ServerAdminEndpointId) const
 {
+	LLM_SCOPE_BYTAG(Concert_ConcertClient);
+
 	FConcertAdmin_GetSessionRepositoriesRequest GetRepositoryRequest;
 	return ClientAdminEndpoint->SendRequest<FConcertAdmin_GetSessionRepositoriesRequest, FConcertAdmin_GetSessionRepositoriesResponse>(GetRepositoryRequest, ServerAdminEndpointId);
 }
 
 TFuture<FConcertAdmin_DropSessionRepositoriesResponse> FConcertClient::DropSessionRepositories(const FGuid& ServerAdminEndpointId, const TArray<FGuid>& RepositoryIds) const
 {
+	LLM_SCOPE_BYTAG(Concert_ConcertClient);
+
 	FConcertAdmin_DropSessionRepositoriesRequest DropRepositoryRequest;
 	DropRepositoryRequest.RepositoryIds = RepositoryIds;
 	return ClientAdminEndpoint->SendRequest<FConcertAdmin_DropSessionRepositoriesRequest, FConcertAdmin_DropSessionRepositoriesResponse>(DropRepositoryRequest, ServerAdminEndpointId);
@@ -1124,24 +1139,32 @@ TFuture<FConcertAdmin_DropSessionRepositoriesResponse> FConcertClient::DropSessi
 
 TFuture<FConcertAdmin_GetAllSessionsResponse> FConcertClient::GetServerSessions(const FGuid& ServerAdminEndpointId) const
 {
+	LLM_SCOPE_BYTAG(Concert_ConcertClient);
+
 	FConcertAdmin_GetAllSessionsRequest GetSessionsRequest = FConcertAdmin_GetAllSessionsRequest();
 	return ClientAdminEndpoint->SendRequest<FConcertAdmin_GetAllSessionsRequest, FConcertAdmin_GetAllSessionsResponse>(GetSessionsRequest, ServerAdminEndpointId);
 }
 
 TFuture<FConcertAdmin_GetSessionsResponse> FConcertClient::GetLiveSessions(const FGuid& ServerAdminEndpointId) const
 {
+	LLM_SCOPE_BYTAG(Concert_ConcertClient);
+
 	FConcertAdmin_GetLiveSessionsRequest GetLiveSessionsRequest;
 	return ClientAdminEndpoint->SendRequest<FConcertAdmin_GetLiveSessionsRequest, FConcertAdmin_GetSessionsResponse>(GetLiveSessionsRequest, ServerAdminEndpointId);
 }
 
 TFuture<FConcertAdmin_GetSessionsResponse> FConcertClient::GetArchivedSessions(const FGuid& ServerAdminEndpointId) const
 {
+	LLM_SCOPE_BYTAG(Concert_ConcertClient);
+
 	FConcertAdmin_GetArchivedSessionsRequest GetArchivedSessionsRequest;
 	return ClientAdminEndpoint->SendRequest<FConcertAdmin_GetArchivedSessionsRequest, FConcertAdmin_GetSessionsResponse>(GetArchivedSessionsRequest, ServerAdminEndpointId);
 }
 
 TFuture<FConcertAdmin_GetSessionClientsResponse> FConcertClient::GetSessionClients(const FGuid& ServerAdminEndpointId, const FGuid& SessionId) const
 {
+	LLM_SCOPE_BYTAG(Concert_ConcertClient);
+
 	FConcertAdmin_GetSessionClientsRequest GetSessionClientsRequest;
 	GetSessionClientsRequest.SessionId = SessionId;
 	return ClientAdminEndpoint->SendRequest<FConcertAdmin_GetSessionClientsRequest, FConcertAdmin_GetSessionClientsResponse>(GetSessionClientsRequest, ServerAdminEndpointId);
@@ -1149,6 +1172,8 @@ TFuture<FConcertAdmin_GetSessionClientsResponse> FConcertClient::GetSessionClien
 
 TFuture<FConcertAdmin_GetSessionActivitiesResponse> FConcertClient::GetSessionActivities(const FGuid& ServerAdminEndpointId, const FGuid& SessionId, int64 FromActivityId, int64 ActivityCount, bool bIncludeDetails) const
 {
+	LLM_SCOPE_BYTAG(Concert_ConcertClient);
+
 	FConcertAdmin_GetSessionActivitiesRequest GetSessionActivitiesRequest;
 	GetSessionActivitiesRequest.SessionId = SessionId;
 	GetSessionActivitiesRequest.FromActivityId = FromActivityId;
diff --git a/Engine/Plugins/Developer/Concert/ConcertMain/Source/Concert/Private/ConcertVersion.cpp b/Engine/Plugins/Developer/Concert/ConcertMain/Source/Concert/Private/ConcertVersion.cpp
--- a/Engine/Plugins/Developer/Concert/ConcertMain/Source/Concert/Private/ConcertVersion.cpp
+++ b/Engine/Plugins/Developer/Concert/ConcertMain/Source/Concert/Private/ConcertVersion.cpp
@@ -7,6 +7,7 @@
 #include "Modules/BuildVersion.h"
 #include "Serialization/CustomVersion.h"
 
+LLM_DEFINE_TAG(Concert_ConcertVersion);
 #define LOCTEXT_NAMESPACE "ConcertVersion"
 
 namespace ConcertVersionUtil
@@ -97,6 +98,8 @@ bool FConcertCustomVersionInfo::Validate(const FConcertCustomVersionInfo& InOthe
 
 void FConcertSessionVersionInfo::Initialize(bool bSupportMixedBuildTypes)
 {
+	LLM_SCOPE_BYTAG(Concert_ConcertVersion);
+
 	FileVersion.Initialize();
 	EngineVersion.Initialize(FEngineVersion::Current());
 
diff --git a/Engine/Plugins/Developer/Concert/ConcertMain/Source/ConcertTransport/Private/ConcertLocalEndpoint.cpp b/Engine/Plugins/Developer/Concert/ConcertMain/Source/ConcertTransport/Private/ConcertLocalEndpoint.cpp
--- a/Engine/Plugins/Developer/Concert/ConcertMain/Source/ConcertTransport/Private/ConcertLocalEndpoint.cpp
+++ b/Engine/Plugins/Developer/Concert/ConcertMain/Source/ConcertTransport/Private/ConcertLocalEndpoint.cpp
@@ -554,6 +554,7 @@ void FConcertLocalEndpoint::ProcessKeepAliveMessage(const TSharedPtr<IMessageCon
 
 void FConcertLocalEndpoint::InternalHandleMessage(const TSharedRef<IMessageContext, ESPMode::ThreadSafe>& Context)
 {
+	LLM_SCOPE_BYTAG(Concert_ConcertLocalEndpoint);
 	SCOPED_CONCERT_TRACE(FConcertLocalEndpoint_InternalHandleMessage);
 	const UScriptStruct* MessageTypeInfo = Context->GetMessageTypeInfo().Get();
 
diff --git a/Engine/Plugins/Developer/Concert/ConcertMain/Source/ConcertTransport/Private/ConcertTransportModule.cpp b/Engine/Plugins/Developer/Concert/ConcertMain/Source/ConcertTransport/Private/ConcertTransportModule.cpp
--- a/Engine/Plugins/Developer/Concert/ConcertMain/Source/ConcertTransport/Private/ConcertTransportModule.cpp
+++ b/Engine/Plugins/Developer/Concert/ConcertMain/Source/ConcertTransport/Private/ConcertTransportModule.cpp
@@ -21,6 +21,7 @@ class FConcertTransportModule : public IConcertTransportModule
 public:
 	virtual TSharedPtr<IConcertEndpointProvider> CreateEndpointProvider() const override
 	{
+		LLM_SCOPE_BYNAME(TEXT("Concert/ConcertTransport"));
 		return MakeShared<FConcertEndpointProvider>();
 	}
 };
diff --git a/Engine/Plugins/Developer/Concert/ConcertSync/ConcertSyncClient/Source/ConcertSyncClient/Private/ConcertClientPackageBridge.cpp b/Engine/Plugins/Developer/Concert/ConcertSync/ConcertSyncClient/Source/ConcertSyncClient/Private/ConcertClientPackageBridge.cpp
--- a/Engine/Plugins/Developer/Concert/ConcertSync/ConcertSyncClient/Source/ConcertSyncClient/Private/ConcertClientPackageBridge.cpp
+++ b/Engine/Plugins/Developer/Concert/ConcertSync/ConcertSyncClient/Source/ConcertSyncClient/Private/ConcertClientPackageBridge.cpp
@@ -26,6 +26,7 @@
 	#include "LevelEditor.h"
 #endif
 
+LLM_DEFINE_TAG(Concert_ConcertClientPackageBridge);
 #define LOCTEXT_NAMESPACE "ConcertClientPackageBridge"
 
 namespace ConcertClientPackageBridgeUtil
@@ -54,6 +55,7 @@ FConcertClientPackageBridge::FConcertClientPackageBridge()
 	: bIgnoreLocalSave(false)
 	, bIgnoreLocalDiscard(false)
 {
+	LLM_SCOPE_BYTAG(Concert_ConcertClientPackageBridge);
 #if WITH_EDITOR
 	// Previously these event handlers were wrapped in the GIsEditor below.  We moved them out to support take recording
 	// on nodes running in -game mode. This is a very focused use case and it is not safe to rely on these working
diff --git a/Engine/Plugins/Developer/Concert/ConcertSync/ConcertSyncClient/Source/ConcertSyncClient/Private/ConcertClientTransactionBridge.cpp b/Engine/Plugins/Developer/Concert/ConcertSync/ConcertSyncClient/Source/ConcertSyncClient/Private/ConcertClientTransactionBridge.cpp
--- a/Engine/Plugins/Developer/Concert/ConcertSync/ConcertSyncClient/Source/ConcertSyncClient/Private/ConcertClientTransactionBridge.cpp
+++ b/Engine/Plugins/Developer/Concert/ConcertSync/ConcertSyncClient/Source/ConcertSyncClient/Private/ConcertClientTransactionBridge.cpp
@@ -24,7 +24,7 @@
 	#include "Elements/Framework/TypedElementSelectionSet.h"
 #endif
 
-
+LLM_DEFINE_TAG(Concert_ConcertClientTransactionBridge);
 #define LOCTEXT_NAMESPACE "ConcertClientTransactionBridge"
 
 namespace ConcertClientTransactionBridgeUtil
@@ -726,6 +726,7 @@ void FConcertClientTransactionBridge::HandleTransactionStateChanged(const FTrans
 
 void FConcertClientTransactionBridge::HandleObjectTransacted(UObject* InObject, const FTransactionObjectEvent& InTransactionEvent)
 {
+	LLM_SCOPE_BYTAG(Concert_ConcertClientTransactionBridge);
 	SCOPED_CONCERT_TRACE(FConcertClientTransactionBridge_HandleObjectTransacted);
 
 	if (bIgnoreLocalTransactions)
@@ -933,6 +934,8 @@ void FConcertClientTransactionBridge::HandleObjectTransacted(UObject* InObject,
 
 void FConcertClientTransactionBridge::ConditionalBindUnderlyingLocalTransactionEvents()
 {
+	LLM_SCOPE_BYTAG(Concert_ConcertClientTransactionBridge);
+
 	if (bHasBoundUnderlyingLocalTransactionEvents)
 	{
 		return;
@@ -998,6 +1001,8 @@ void FConcertClientTransactionBridge::OnEndFrame()
 
 void FConcertClientTransactionBridge::RegisterTransactionFilter(FName FilterName, FTransactionFilterDelegate FilterHandle)
 {
+	LLM_SCOPE_BYTAG(Concert_ConcertClientTransactionBridge);
+
 	check(TransactionFilters.Find(FilterName) == nullptr);
 
 	TransactionFilters.Add(FilterName) = MoveTemp(FilterHandle);
diff --git a/Engine/Plugins/Developer/Concert/ConcertSync/ConcertSyncClient/Source/ConcertSyncClient/Private/ConcertClientWorkspace.cpp b/Engine/Plugins/Developer/Concert/ConcertSync/ConcertSyncClient/Source/ConcertSyncClient/Private/ConcertClientWorkspace.cpp
--- a/Engine/Plugins/Developer/Concert/ConcertSync/ConcertSyncClient/Source/ConcertSyncClient/Private/ConcertClientWorkspace.cpp
+++ b/Engine/Plugins/Developer/Concert/ConcertSync/ConcertSyncClient/Source/ConcertSyncClient/Private/ConcertClientWorkspace.cpp
@@ -63,6 +63,7 @@
 	#include "GameMapsSettings.h"
 #endif
 
+LLM_DEFINE_TAG(Concert_ConcertClientWorkspace);
 #define LOCTEXT_NAMESPACE "ConcertClientWorkspace"
 
 /** Provides a workaround for scoped slow tasks that are push/pop out of order by Concert. Concert doesn't use FScopedSlowTask as designed when syncing
@@ -817,6 +818,7 @@ void FConcertClientWorkspace::HandleEndPIE(const bool InIsSimulating)
 
 void FConcertClientWorkspace::OnEndFrame()
 {
+	LLM_SCOPE_BYTAG(Concert_ConcertClientWorkspace);
 	SCOPED_CONCERT_TRACE(FConcertClientWorkspace_OnEndFrame);
 
 	if (CanFinalize())
diff --git a/Engine/Plugins/Developer/Concert/ConcertSync/ConcertSyncClient/Source/ConcertSyncClient/Private/ConcertSyncClient.cpp b/Engine/Plugins/Developer/Concert/ConcertSync/ConcertSyncClient/Source/ConcertSyncClient/Private/ConcertSyncClient.cpp
--- a/Engine/Plugins/Developer/Concert/ConcertSync/ConcertSyncClient/Source/ConcertSyncClient/Private/ConcertSyncClient.cpp
+++ b/Engine/Plugins/Developer/Concert/ConcertSync/ConcertSyncClient/Source/ConcertSyncClient/Private/ConcertSyncClient.cpp
@@ -11,8 +11,8 @@
 #include "ConcertClientPresenceManager.h"
 #include "ConcertSourceControlProxy.h"
 
+LLM_DECLARE_TAG(Concert_ConcertSyncClient);
 #define LOCTEXT_NAMESPACE "ConcertSyncClient"
-LLM_DEFINE_TAG(Concert_ConcertSyncClient);
 
 FConcertSyncClient::FConcertSyncClient(const FString& InRole, IConcertClientPackageBridge* InPackageBridge, IConcertClientTransactionBridge* InTransactionBridge)
 	: ConcertClient(IConcertModule::Get().CreateClient(InRole))
diff --git a/Engine/Plugins/Developer/Concert/ConcertSync/ConcertSyncClient/Source/ConcertSyncClient/Private/ConcertSyncClientModule.cpp b/Engine/Plugins/Developer/Concert/ConcertSync/ConcertSyncClient/Source/ConcertSyncClient/Private/ConcertSyncClientModule.cpp
--- a/Engine/Plugins/Developer/Concert/ConcertSync/ConcertSyncClient/Source/ConcertSyncClient/Private/ConcertSyncClientModule.cpp
+++ b/Engine/Plugins/Developer/Concert/ConcertSync/ConcertSyncClient/Source/ConcertSyncClient/Private/ConcertSyncClientModule.cpp
@@ -10,6 +10,8 @@
 
 #include "Framework/Notifications/NotificationManager.h"
 
+LLM_DEFINE_TAG(Concert_ConcertSyncClient);
+
 /**
  * Implements the Concert Sync module for Event synchronization
  */
@@ -21,18 +23,22 @@ public:
 
 	virtual void StartupModule() override
 	{
+		LLM_SCOPE_BYTAG(Concert_ConcertSyncClient);
 		PackageBridge = MakeUnique<FConcertClientPackageBridge>();
 		TransactionBridge = MakeUnique<FConcertClientTransactionBridge>();
 	}
 
 	virtual void ShutdownModule() override
 	{
+		LLM_SCOPE_BYTAG(Concert_ConcertSyncClient);
 		PackageBridge.Reset();
 		TransactionBridge.Reset();
 	}
 
 	virtual UConcertClientConfig* ParseClientSettings(const TCHAR* CommandLine) override
 	{
+		LLM_SCOPE_BYTAG(Concert_ConcertSyncClient);
+
 		UConcertClientConfig* ClientConfig = NewObject<UConcertClientConfig>();
 
 		if (ClientConfig)
@@ -101,6 +107,8 @@ public:
 
 	virtual TSharedRef<IConcertSyncClient> CreateClient(const FString& InRole) override
 	{
+		LLM_SCOPE_BYTAG(Concert_ConcertSyncClient);
+
 		// Remove dead clients.
 		Clients.RemoveAll([](TWeakPtr<IConcertSyncClient> WeakClient) { return !WeakClient.IsValid(); });
 
diff --git a/Engine/Plugins/Experimental/RemoteSession/Source/RemoteSessionEditor/Private/RemoteSessionEditorModule.cpp b/Engine/Plugins/Experimental/RemoteSession/Source/RemoteSessionEditor/Private/RemoteSessionEditorModule.cpp
--- a/Engine/Plugins/Experimental/RemoteSession/Source/RemoteSessionEditor/Private/RemoteSessionEditorModule.cpp
+++ b/Engine/Plugins/Experimental/RemoteSession/Source/RemoteSessionEditor/Private/RemoteSessionEditorModule.cpp
@@ -6,21 +6,24 @@
 #include "Modules/ModuleManager.h"
 #include "Widgets/SRemoteSessionStream.h"
 
+LLM_DEFINE_TAG(RemoteSession_RemoteSessionEditor);
 #define LOCTEXT_NAMESPACE "FRemoteSessionEditorModule"
 
-
 class FRemoteSessionEditorModule : public IModuleInterface
 {
 public:
 	virtual void StartupModule() override
 	{
-		LLM_SCOPE_BYNAME(TEXT("RemoteSessionEditor"));
+		LLM_SCOPE_BYTAG(RemoteSession_RemoteSessionEditor);
+
 		FRemoteSessionEditorStyle::Register();
 		SRemoteSessionStream::RegisterNomadTabSpawner();
 	}
 
 	virtual void ShutdownModule() override
 	{
+		LLM_SCOPE_BYTAG(RemoteSession_RemoteSessionEditor);
+
 		if (!IsEngineExitRequested() && UObjectInitialized())
 		{
 			FRemoteSessionEditorStyle::Unregister();
diff --git a/Engine/Plugins/Experimental/RemoteSession/Source/RemoteSessionEditor/Private/Widgets/SRemoteSessionStream.cpp b/Engine/Plugins/Experimental/RemoteSession/Source/RemoteSessionEditor/Private/Widgets/SRemoteSessionStream.cpp
--- a/Engine/Plugins/Experimental/RemoteSession/Source/RemoteSessionEditor/Private/Widgets/SRemoteSessionStream.cpp
+++ b/Engine/Plugins/Experimental/RemoteSession/Source/RemoteSessionEditor/Private/Widgets/SRemoteSessionStream.cpp
@@ -73,6 +73,7 @@ void SRemoteSessionStream::RegisterNomadTabSpawner()
 {
 	auto RegisterTabSpawner = []()
 	{
+		LLM_SCOPE_BYNAME(TEXT("RemoteSession/RemoteSessionStream"));
 		FLevelEditorModule& LevelEditorModule = FModuleManager::LoadModuleChecked<FLevelEditorModule>(RemoteSessionStream::LevelEditorModuleName);
 		TSharedPtr<FTabManager> LevelEditorTabManager = LevelEditorModule.GetLevelEditorTabManager();
 
diff --git a/Engine/Plugins/Experimental/VirtualCamera/Source/VCamInput/Private/VCamInputSubsystem.cpp b/Engine/Plugins/Experimental/VirtualCamera/Source/VCamInput/Private/VCamInputSubsystem.cpp
--- a/Engine/Plugins/Experimental/VirtualCamera/Source/VCamInput/Private/VCamInputSubsystem.cpp
+++ b/Engine/Plugins/Experimental/VirtualCamera/Source/VCamInput/Private/VCamInputSubsystem.cpp
@@ -5,10 +5,13 @@
 
 #include "Framework/Application/SlateApplication.h"
 
+LLM_DEFINE_TAG(VirtualCamera_CamInputSubsystem);
 DEFINE_LOG_CATEGORY(LogVCamInput);
 
 UVCamInputSubsystem::UVCamInputSubsystem()
 {
+	LLM_SCOPE_BYTAG(VirtualCamera_CamInputSubsystem);
+
 	// Registering the input processor is only valid in the actual subsystem and not the CDO
 	if (!HasAnyFlags(RF_ClassDefaultObject))
 	{
diff --git a/Engine/Plugins/Experimental/VirtualCamera/Source/VirtualCamera/Private/VirtualCameraMovementComponent.cpp b/Engine/Plugins/Experimental/VirtualCamera/Source/VirtualCamera/Private/VirtualCameraMovementComponent.cpp
--- a/Engine/Plugins/Experimental/VirtualCamera/Source/VirtualCamera/Private/VirtualCameraMovementComponent.cpp
+++ b/Engine/Plugins/Experimental/VirtualCamera/Source/VirtualCamera/Private/VirtualCameraMovementComponent.cpp
@@ -5,6 +5,8 @@
 #include "Components/PrimitiveComponent.h"
 #include "Kismet/KismetMathLibrary.h"
 
+LLM_DEFINE_TAG(VirtualCamera_VirtualCameraMovementComponent);
+
 namespace
 {
 	float ClampStabilizationScale(float StabilizationScale)
@@ -25,6 +27,8 @@ namespace
 
 UVirtualCameraMovementComponent::UVirtualCameraMovementComponent(const FObjectInitializer& ObjectInitializer)
 {
+	LLM_SCOPE_BYTAG(VirtualCamera_VirtualCameraMovementComponent);
+
 	// Initialize Axis Settings Map
 	AxisSettings.Add(EVirtualCameraAxis::LocationX);
 	AxisSettings.Add(EVirtualCameraAxis::LocationY);
diff --git a/Engine/Plugins/Experimental/VirtualCamera/Source/VirtualCamera/Private/VirtualCameraPawnBase.cpp b/Engine/Plugins/Experimental/VirtualCamera/Source/VirtualCamera/Private/VirtualCameraPawnBase.cpp
--- a/Engine/Plugins/Experimental/VirtualCamera/Source/VirtualCamera/Private/VirtualCameraPawnBase.cpp
+++ b/Engine/Plugins/Experimental/VirtualCamera/Source/VirtualCamera/Private/VirtualCameraPawnBase.cpp
@@ -5,12 +5,15 @@
 #include "UnrealEngine.h"
 #include "Kismet/GameplayStatics.h"
 
+LLM_DEFINE_TAG(VirtualCamera_VirtualCameraPawnBase);
+
 int32 AVirtualCameraPawnBase::PresetIndex = 1;
 int32 AVirtualCameraPawnBase::WaypointIndex = 1;
 int32 AVirtualCameraPawnBase::ScreenshotIndex = 1;
 
 AVirtualCameraPawnBase::AVirtualCameraPawnBase(const FObjectInitializer& ObjectInitializer)
 {
+	LLM_SCOPE_BYTAG(VirtualCamera_VirtualCameraPawnBase);
 	PrimaryActorTick.bCanEverTick = true;
 	PrimaryActorTick.bStartWithTickEnabled = true;
 
diff --git a/Engine/Plugins/Experimental/VirtualCamera/Source/VirtualCamera/Private/VirtualCameraSaveGame.cpp b/Engine/Plugins/Experimental/VirtualCamera/Source/VirtualCamera/Private/VirtualCameraSaveGame.cpp
--- a/Engine/Plugins/Experimental/VirtualCamera/Source/VirtualCamera/Private/VirtualCameraSaveGame.cpp
+++ b/Engine/Plugins/Experimental/VirtualCamera/Source/VirtualCamera/Private/VirtualCameraSaveGame.cpp
@@ -8,6 +8,7 @@ int32 FVirtualCameraScreenshot::NextIndex = 1;
 
 UVirtualCameraSaveGame::UVirtualCameraSaveGame(const FObjectInitializer& ObjectInitializer)
 {
+	LLM_SCOPE_BYNAME("VirtualCamera/VirtualCameraSaveGame");
 	SaveSlotName = "SavedVirtualCameraSettings";
 	UserIndex = 0;
 }
\ No newline at end of file
diff --git a/Engine/Plugins/Experimental/VirtualProductionUtilities/Source/VPBookmarkEditor/Private/VPBookmarkEditorModule.cpp b/Engine/Plugins/Experimental/VirtualProductionUtilities/Source/VPBookmarkEditor/Private/VPBookmarkEditorModule.cpp
--- a/Engine/Plugins/Experimental/VirtualProductionUtilities/Source/VPBookmarkEditor/Private/VPBookmarkEditorModule.cpp
+++ b/Engine/Plugins/Experimental/VirtualProductionUtilities/Source/VPBookmarkEditor/Private/VPBookmarkEditorModule.cpp
@@ -7,11 +7,14 @@
 #include "Bookmarks/IBookmarkTypeTools.h"
 
 
+LLM_DEFINE_TAG(VirtualProduction_VPBookmarkEditor);
 DEFINE_LOG_CATEGORY(LogVPBookmarkEditor);
 
 
 void FVPBookmarkEditorModule::StartupModule()
 {
+	LLM_SCOPE_BYTAG(VirtualProduction_VPBookmarkEditor);
+
 	BookmarkTypeActions = MakeShared<FVPBookmarkTypeActions>();
 	IBookmarkTypeTools::Get().RegisterBookmarkTypeActions(BookmarkTypeActions.ToSharedRef());
 }
@@ -19,6 +22,8 @@ void FVPBookmarkEditorModule::StartupModule()
 
 void FVPBookmarkEditorModule::ShutdownModule()
 {
+	LLM_SCOPE_BYTAG(VirtualProduction_VPBookmarkEditor);
+
 	IBookmarkTypeTools::Get().UnregisterBookmarkTypeActions(BookmarkTypeActions.ToSharedRef());
 	BookmarkTypeActions.Reset();
 }
diff --git a/Engine/Plugins/Experimental/VirtualProductionUtilities/Source/VPMaterialsEditor/Private/VPMaterialsEditorModule.cpp b/Engine/Plugins/Experimental/VirtualProductionUtilities/Source/VPMaterialsEditor/Private/VPMaterialsEditorModule.cpp
--- a/Engine/Plugins/Experimental/VirtualProductionUtilities/Source/VPMaterialsEditor/Private/VPMaterialsEditorModule.cpp
+++ b/Engine/Plugins/Experimental/VirtualProductionUtilities/Source/VPMaterialsEditor/Private/VPMaterialsEditorModule.cpp
@@ -5,18 +5,23 @@
 
 #include "MaterialList.h"
 
+LLM_DEFINE_TAG(VirtualProduction_VPMaterialsEditor);
+
 void FVPMaterialsEditorModule::StartupModule()
 {
+	LLM_SCOPE_BYTAG(VirtualProduction_VPMaterialsEditor);
+
 	// Add bottom extender for material item
 	FMaterialList::OnAddMaterialItemViewExtraBottomWidget.AddLambda([](const TSharedRef<FMaterialItemView>& InMaterialItemView, UActorComponent* InCurrentComponent, IDetailLayoutBuilder& InDetailBuilder, TArray<TSharedPtr<SWidget>>& OutExtensions)
 	{
-		LLM_SCOPE_BYNAME("VirtualProduction/VPMaterialsEditor");
+		LLM_SCOPE_BYTAG(VirtualProduction_VPMaterialsEditor);
 		OutExtensions.Add(SNew(SMaterialDynamicView, InMaterialItemView, InCurrentComponent));
 	});
 }
 
 void FVPMaterialsEditorModule::ShutdownModule()
 {
+	LLM_SCOPE_BYTAG(VirtualProduction_VPMaterialsEditor);
 	FMaterialList::OnAddMaterialItemViewExtraBottomWidget.RemoveAll(this);
 }
 
diff --git a/Engine/Plugins/Experimental/VirtualProductionUtilities/Source/VPUtilitiesEditor/Private/STimecodeProviderTab.cpp b/Engine/Plugins/Experimental/VirtualProductionUtilities/Source/VPUtilitiesEditor/Private/STimecodeProviderTab.cpp
--- a/Engine/Plugins/Experimental/VirtualProductionUtilities/Source/VPUtilitiesEditor/Private/STimecodeProviderTab.cpp
+++ b/Engine/Plugins/Experimental/VirtualProductionUtilities/Source/VPUtilitiesEditor/Private/STimecodeProviderTab.cpp
@@ -20,7 +20,7 @@
 #include "WorkspaceMenuStructure.h"
 #include "WorkspaceMenuStructureModule.h"
 
-
+LLM_DEFINE_TAG(VirtualProductionUtilities_TimecodeProviderTab);
 #define LOCTEXT_NAMESPACE "TimecodeProviderTab"
 
 namespace TimecodeProviderTab
@@ -40,8 +40,11 @@ namespace TimecodeProviderTab
 
 void STimecodeProviderTab::RegisterNomadTabSpawner()
 {
+	LLM_SCOPE_BYTAG(VirtualProductionUtilities_TimecodeProviderTab);
+
 	auto RegisterTabSpawner = []()
 	{
+		LLM_SCOPE_BYTAG(VirtualProductionUtilities_TimecodeProviderTab);
 		FLevelEditorModule& LevelEditorModule = FModuleManager::LoadModuleChecked<FLevelEditorModule>("LevelEditor");
 		TSharedPtr<FTabManager> LevelEditorTabManager = LevelEditorModule.GetLevelEditorTabManager();
 
diff --git a/Engine/Plugins/Experimental/VirtualProductionUtilities/Source/VPUtilitiesEditor/Private/VPScoutingSubsystem.cpp b/Engine/Plugins/Experimental/VirtualProductionUtilities/Source/VPUtilitiesEditor/Private/VPScoutingSubsystem.cpp
--- a/Engine/Plugins/Experimental/VirtualProductionUtilities/Source/VPUtilitiesEditor/Private/VPScoutingSubsystem.cpp
+++ b/Engine/Plugins/Experimental/VirtualProductionUtilities/Source/VPUtilitiesEditor/Private/VPScoutingSubsystem.cpp
@@ -18,6 +18,8 @@
 #include "VPUtilitiesEditorSettings.h"
 #include "LevelEditorActions.h"
 
+LLM_DEFINE_TAG(VirtualProductionUtilities_VPScoutingSubsystem);
+
 /* UVPScoutingSubsystem name
  *****************************************************************************/
 const FName UVPScoutingSubsystem::VProdPanelID = FName(TEXT("VirtualProductionPanel"));
@@ -33,6 +35,8 @@ const FName UVPScoutingSubsystem::VProdPanelGafferID = FName(TEXT("VirtualProduc
  *****************************************************************************/
 UVPScoutingSubsystemGestureManagerBase::UVPScoutingSubsystemGestureManagerBase()
 {
+	LLM_SCOPE_BYTAG(VirtualProductionUtilities_VPScoutingSubsystem);
+
 	if (!HasAnyFlags(RF_ArchetypeObject | RF_ClassDefaultObject))
 	{
 		IVREditorModule::Get().OnVREditingModeEnter().AddUObject(this, &UVPScoutingSubsystemGestureManagerBase::OnVREditingModeEnterCallback);
@@ -42,6 +46,8 @@ UVPScoutingSubsystemGestureManagerBase::UVPScoutingSubsystemGestureManagerBase()
 
 void UVPScoutingSubsystemGestureManagerBase::BeginDestroy()
 {
+	LLM_SCOPE_BYTAG(VirtualProductionUtilities_VPScoutingSubsystem);
+
 	if (!HasAnyFlags(RF_ArchetypeObject | RF_ClassDefaultObject))
 	{
 		IVREditorModule::Get().OnVREditingModeEnter().RemoveAll(this);
@@ -132,6 +138,8 @@ void UVPScoutingSubsystem::Deinitialize()
 
 void UVPScoutingSubsystem::OnEngineInitComplete()
 {
+	LLM_SCOPE_BYTAG(VirtualProductionUtilities_VPScoutingSubsystem);
+
 	FCoreDelegates::OnFEngineLoopInitComplete.Remove(EngineInitCompleteDelegate);
 	EngineInitCompleteDelegate.Reset();
 
diff --git a/Engine/Plugins/Experimental/VirtualProductionUtilities/Source/VPUtilitiesEditor/Private/VPUIBase.cpp b/Engine/Plugins/Experimental/VirtualProductionUtilities/Source/VPUtilitiesEditor/Private/VPUIBase.cpp
--- a/Engine/Plugins/Experimental/VirtualProductionUtilities/Source/VPUtilitiesEditor/Private/VPUIBase.cpp
+++ b/Engine/Plugins/Experimental/VirtualProductionUtilities/Source/VPUtilitiesEditor/Private/VPUIBase.cpp
@@ -13,11 +13,14 @@
 #include "Editor/EditorEngine.h"
 #include "GameFramework/Actor.h"
 
+LLM_DEFINE_TAG(VirtualProductionUtilities_VPUIBase);
 #define LOCTEXT_NAMESPACE "VPUIBase"
 
 
 bool UVPUIBase::Initialize()
 {
+	LLM_SCOPE_BYTAG(VirtualProductionUtilities_VPUIBase);
+
 	const bool SuperInitialized = Super::Initialize();
 
 	if (SuperInitialized && !HasAnyFlags(RF_ClassDefaultObject))
@@ -43,6 +46,8 @@ bool UVPUIBase::Initialize()
 
 void UVPUIBase::BeginDestroy()
 {
+	LLM_SCOPE_BYTAG(VirtualProductionUtilities_VPUIBase);
+
 	if (!HasAnyFlags(RF_ClassDefaultObject))
 	{
 		USelection::SelectObjectEvent.RemoveAll(this);
diff --git a/Engine/Plugins/Experimental/VirtualProductionUtilities/Source/VPUtilitiesEditor/Private/VPUtilitiesEditorModule.cpp b/Engine/Plugins/Experimental/VirtualProductionUtilities/Source/VPUtilitiesEditor/Private/VPUtilitiesEditorModule.cpp
--- a/Engine/Plugins/Experimental/VirtualProductionUtilities/Source/VPUtilitiesEditor/Private/VPUtilitiesEditorModule.cpp
+++ b/Engine/Plugins/Experimental/VirtualProductionUtilities/Source/VPUtilitiesEditor/Private/VPUtilitiesEditorModule.cpp
@@ -26,7 +26,7 @@
 #include "WorkspaceMenuStructure.h"
 #include "WorkspaceMenuStructureModule.h"
 
-
+LLM_DEFINE_TAG(VirtualProductionUtilities_VPUtilitiesEditor);
 #define LOCTEXT_NAMESPACE "VPUtilitiesEditor"
 
 DEFINE_LOG_CATEGORY(LogVPUtilitiesEditor);
@@ -36,6 +36,8 @@ const FName FVPUtilitiesEditorModule::PlacementModeCategoryHandle = TEXT("Virtua
 
 void FVPUtilitiesEditorModule::StartupModule()
 {
+	LLM_SCOPE_BYTAG(VirtualProductionUtilities_VPUtilitiesEditor);
+
 	FVPUtilitiesEditorStyle::Register();
 
 	CustomUIHandler.Reset(NewObject<UVPCustomUIHandler>());
@@ -54,6 +56,8 @@ void FVPUtilitiesEditorModule::StartupModule()
 
 void FVPUtilitiesEditorModule::ShutdownModule()
 {
+	LLM_SCOPE_BYTAG(VirtualProductionUtilities_VPUtilitiesEditor);
+
 	UnregisterSettings();
 	STimecodeProviderTab::UnregisterNomadTabSpawner();
 	SGenlockProviderTab::UnregisterNomadTabSpawner();
diff --git a/Engine/Plugins/Runtime/ActorLayerUtilities/Source/ActorLayerUtilitiesEditor/Private/ActorLayerUtilitiesEditor.cpp b/Engine/Plugins/Runtime/ActorLayerUtilities/Source/ActorLayerUtilitiesEditor/Private/ActorLayerUtilitiesEditor.cpp
--- a/Engine/Plugins/Runtime/ActorLayerUtilities/Source/ActorLayerUtilitiesEditor/Private/ActorLayerUtilitiesEditor.cpp
+++ b/Engine/Plugins/Runtime/ActorLayerUtilities/Source/ActorLayerUtilitiesEditor/Private/ActorLayerUtilitiesEditor.cpp
@@ -5,6 +5,7 @@
 #include "ActorLayerUtilities.h"
 #include "ActorLayerPropertyTypeCustomization.h"
 
+LLM_DEFINE_TAG(ActionLayerUtilities_ActorLayerUtilitiesEditor);
 
 class FActorLayerUtilitiesEditorModule : public IModuleInterface
 {
@@ -17,6 +18,8 @@ class FActorLayerUtilitiesEditorModule : public IModuleInterface
 
 	virtual void StartupModule() override
 	{
+		LLM_SCOPE_BYTAG(ActionLayerUtilities_ActorLayerUtilitiesEditor);
+
 		ActorLayerTypeName = FActorLayer::StaticStruct()->GetFName();
 
 		FPropertyEditorModule& PropertyModule = FModuleManager::LoadModuleChecked<FPropertyEditorModule>("PropertyEditor");
@@ -25,6 +28,8 @@ class FActorLayerUtilitiesEditorModule : public IModuleInterface
 
 	virtual void ShutdownModule() override
 	{
+		LLM_SCOPE_BYTAG(ActionLayerUtilities_ActorLayerUtilitiesEditor);
+
 		FPropertyEditorModule* PropertyModule = FModuleManager::GetModulePtr<FPropertyEditorModule>("PropertyEditor");
 
 		if (PropertyModule && ActorLayerTypeName != NAME_None)
diff --git a/Engine/Plugins/Runtime/Steam/SteamVR/Source/SteamVREditor/Private/SteamVREditor.cpp b/Engine/Plugins/Runtime/Steam/SteamVR/Source/SteamVREditor/Private/SteamVREditor.cpp
--- a/Engine/Plugins/Runtime/Steam/SteamVR/Source/SteamVREditor/Private/SteamVREditor.cpp
+++ b/Engine/Plugins/Runtime/Steam/SteamVR/Source/SteamVREditor/Private/SteamVREditor.cpp
@@ -48,12 +48,15 @@ POSSIBILITY OF SUCH DAMAGE.
 
 static const FName SteamVREditorTabName("SteamVREditor");
 
+LLM_DEFINE_TAG(SteamVR_SteamVREditor);
 #define LOCTEXT_NAMESPACE "FSteamVREditorModule"
 
 PRAGMA_DISABLE_DEPRECATION_WARNINGS
 
 void FSteamVREditorModule::StartupModule()
 {
+	LLM_SCOPE_BYTAG(SteamVR_SteamVREditor);
+
 	RegisterSettings();
 
 	FSteamVREditorStyle::Initialize();
@@ -138,6 +141,8 @@ bool FSteamVREditorModule::ShowSteamVRInputToolbarDropdown()
 
 void FSteamVREditorModule::ShutdownModule()
 {
+	LLM_SCOPE_BYTAG(SteamVR_SteamVREditor);
+
 	FSteamVREditorStyle::Shutdown();
 
 	FSteamVREditorCommands::Unregister();
diff --git a/Engine/Plugins/VirtualProduction/CameraCalibrationCore/Source/CameraCalibrationCoreEditor/Private/CameraCalibrationCoreEditorModule.cpp b/Engine/Plugins/VirtualProduction/CameraCalibrationCore/Source/CameraCalibrationCoreEditor/Private/CameraCalibrationCoreEditorModule.cpp
--- a/Engine/Plugins/VirtualProduction/CameraCalibrationCore/Source/CameraCalibrationCoreEditor/Private/CameraCalibrationCoreEditorModule.cpp
+++ b/Engine/Plugins/VirtualProduction/CameraCalibrationCore/Source/CameraCalibrationCoreEditor/Private/CameraCalibrationCoreEditorModule.cpp
@@ -16,6 +16,7 @@
 #include "Modules/ModuleManager.h"
 #include "PropertyEditorModule.h"
 
+LLM_DEFINE_TAG(CameraCalibrationCore_CameraCalibrationCoreEditor);
 #define LOCTEXT_NAMESPACE "CameraCalibrationCoreEditor"
 
 DEFINE_LOG_CATEGORY(LogCameraCalibrationCoreEditor);
@@ -23,6 +24,8 @@ DEFINE_LOG_CATEGORY(LogCameraCalibrationCoreEditor);
 
 void FCameraCalibrationCoreEditorModule::StartupModule()
 {
+	LLM_SCOPE_BYTAG(CameraCalibrationCore_CameraCalibrationCoreEditor);
+
 	FCameraCalibrationCoreEditorStyle::Get();
 
 	FPropertyEditorModule& PropertyEditorModule = FModuleManager::Get().LoadModuleChecked<FPropertyEditorModule>("PropertyEditor");
@@ -37,6 +40,8 @@ void FCameraCalibrationCoreEditorModule::StartupModule()
 
 void FCameraCalibrationCoreEditorModule::ShutdownModule()
 {
+	LLM_SCOPE_BYTAG(CameraCalibrationCore_CameraCalibrationCoreEditor);
+
 	if (!IsEngineExitRequested() && GEditor && UObjectInitialized())
 	{
 		FPropertyEditorModule& PropertyModule = FModuleManager::GetModuleChecked<FPropertyEditorModule>("PropertyEditor");
@@ -93,6 +98,7 @@ void FCameraCalibrationCoreEditorModule::RegisterPlacementModeItems()
 {
 	auto RegisterPlaceActors = [&]() -> void
 	{
+		LLM_SCOPE_BYTAG(CameraCalibrationCore_CameraCalibrationCoreEditor)
 		if (!GEditor)
 		{
 			return;
diff --git a/Engine/Plugins/VirtualProduction/LiveLinkCamera/Source/LiveLinkCameraEditor/Private/LiveLinkCameraEditorModule.cpp b/Engine/Plugins/VirtualProduction/LiveLinkCamera/Source/LiveLinkCameraEditor/Private/LiveLinkCameraEditorModule.cpp
--- a/Engine/Plugins/VirtualProduction/LiveLinkCamera/Source/LiveLinkCameraEditor/Private/LiveLinkCameraEditorModule.cpp
+++ b/Engine/Plugins/VirtualProduction/LiveLinkCamera/Source/LiveLinkCameraEditor/Private/LiveLinkCameraEditorModule.cpp
@@ -6,6 +6,8 @@
 #include "LiveLinkCameraControllerCustomization.h"
 #include "PropertyEditorModule.h"
 
+LLM_DEFINE_TAG(LiveLink_LiveLinkCameraEditor);
+
 void FLiveLinkCameraEditorModule::StartupModule()
 {
 	RegisterCustomizations();
@@ -18,12 +20,16 @@ void FLiveLinkCameraEditorModule::ShutdownModule()
 
 void FLiveLinkCameraEditorModule::RegisterCustomizations()
 {
+	LLM_SCOPE_BYTAG(LiveLink_LiveLinkCameraEditor);
+
 	FPropertyEditorModule& PropertyEditorModule = FModuleManager::Get().LoadModuleChecked<FPropertyEditorModule>("PropertyEditor");
 	PropertyEditorModule.RegisterCustomClassLayout(ULiveLinkCameraController::StaticClass()->GetFName(), FOnGetDetailCustomizationInstance::CreateStatic(&FLiveLinkCameraControllerCustomization::MakeInstance));
 }
 
 void FLiveLinkCameraEditorModule::UnregisterCustomizations()
 {
+	LLM_SCOPE_BYTAG(LiveLink_LiveLinkCameraEditor);
+
 	if (UObjectInitialized() && !IsEngineExitRequested())
 	{
 		FPropertyEditorModule* PropertyEditorModule = FModuleManager::Get().GetModulePtr<FPropertyEditorModule>("PropertyEditor");
diff --git a/Engine/Plugins/VirtualProduction/LiveLinkCamera/Source/LiveLinkCameraRecording/Private/LiveLinkCameraRecording.cpp b/Engine/Plugins/VirtualProduction/LiveLinkCamera/Source/LiveLinkCameraRecording/Private/LiveLinkCameraRecording.cpp
--- a/Engine/Plugins/VirtualProduction/LiveLinkCamera/Source/LiveLinkCameraRecording/Private/LiveLinkCameraRecording.cpp
+++ b/Engine/Plugins/VirtualProduction/LiveLinkCamera/Source/LiveLinkCameraRecording/Private/LiveLinkCameraRecording.cpp
@@ -4,18 +4,24 @@
 #include "LiveLinkCameraControllerTrackEditor.h"
 #include "Modules/ModuleManager.h"
 
+LLM_DEFINE_TAG(LiveLinkCamera_LiveLinkCameraRecording);
+
 class FLiveLinkCameraRecordingModule : public IModuleInterface
 {
 public:
 	// IModuleInterface interface
 	virtual void StartupModule() override
 	{
+		LLM_SCOPE_BYTAG(LiveLinkCamera_LiveLinkCameraRecording);
+
 		ISequencerModule& SequencerModule = FModuleManager::LoadModuleChecked<ISequencerModule>("Sequencer");
 		CreateLiveLinkCameraControllerTrackEditorHandle = SequencerModule.RegisterTrackEditor(FOnCreateTrackEditor::CreateStatic(&FLiveLinkCameraControllerTrackEditor::CreateTrackEditor));
 	}
 
 	virtual void ShutdownModule() override
 	{
+		LLM_SCOPE_BYTAG(LiveLinkCamera_LiveLinkCameraRecording);
+
 		ISequencerModule* SequencerModule = FModuleManager::GetModulePtr<ISequencerModule>("Sequencer");
 		if (SequencerModule != nullptr)
 		{
diff --git a/Engine/Plugins/VirtualProduction/Takes/Source/TakeMovieScene/Private/MovieSceneTakeSection.cpp b/Engine/Plugins/VirtualProduction/Takes/Source/TakeMovieScene/Private/MovieSceneTakeSection.cpp
--- a/Engine/Plugins/VirtualProduction/Takes/Source/TakeMovieScene/Private/MovieSceneTakeSection.cpp
+++ b/Engine/Plugins/VirtualProduction/Takes/Source/TakeMovieScene/Private/MovieSceneTakeSection.cpp
@@ -71,6 +71,8 @@ void UMovieSceneTakeSection::PostEditImport()
 
 void UMovieSceneTakeSection::ReconstructChannelProxy()
 {
+	LLM_SCOPE_BYNAME(TEXT("Takes/MovieSceneTakeSection"))
+
 	FMovieSceneChannelProxyData Channels;
 
 #if WITH_EDITOR
diff --git a/Engine/Plugins/VirtualProduction/Takes/Source/TakesCore/Private/TakesCoreModule.cpp b/Engine/Plugins/VirtualProduction/Takes/Source/TakesCore/Private/TakesCoreModule.cpp
--- a/Engine/Plugins/VirtualProduction/Takes/Source/TakesCore/Private/TakesCoreModule.cpp
+++ b/Engine/Plugins/VirtualProduction/Takes/Source/TakesCore/Private/TakesCoreModule.cpp
@@ -9,10 +9,14 @@
 
 #include "Modules/ModuleManager.h"
 
+LLM_DEFINE_TAG(Takes_TakesCore);
+
 class FTakesCoreModule : public IModuleInterface
 {
 	virtual void StartupModule() override
 	{
+		LLM_SCOPE_BYTAG(Takes_TakesCore);
+
 		// Ensure the level sequence module is loaded
 		FModuleManager::Get().LoadModuleChecked("LevelSequence");
 
@@ -30,6 +34,8 @@ class FTakesCoreModule : public IModuleInterface
 
 	virtual void ShutdownModule() override
 	{
+		LLM_SCOPE_BYTAG(Takes_TakesCore);
+
 		if (ULevelSequence* CDO = LevelSequenceCDO.Get())
 		{
 			CDO->RemoveMetaData<UTakeMetaData>();
