Fix session's last modified property not getting serialized when sent to clients

diff --git a/Engine/Plugins/Developer/Concert/ConcertMain/Source/Concert/Private/ConcertServer.cpp b/Engine/Plugins/Developer/Concert/ConcertMain/Source/Concert/Private/ConcertServer.cpp
--- a/Engine/Plugins/Developer/Concert/ConcertMain/Source/Concert/Private/ConcertServer.cpp
+++ b/Engine/Plugins/Developer/Concert/ConcertMain/Source/Concert/Private/ConcertServer.cpp
@@ -595,7 +595,7 @@ void FConcertServer::UpdateLastModified(TArray<FConcertSessionInfo>& SessionInfo
 {
 	for (int32 i = 0; i < SessionInfos.Num(); ++i)
 	{
-		SessionInfos[i].LastModified = SessionCreationTimes[i];
+		SessionInfos[i].SetLastModified(SessionCreationTimes[i]);
 	}
 }
 
diff --git a/Engine/Plugins/Developer/Concert/ConcertMain/Source/Concert/Public/ConcertMessageData.h b/Engine/Plugins/Developer/Concert/ConcertMain/Source/Concert/Public/ConcertMessageData.h
--- a/Engine/Plugins/Developer/Concert/ConcertMain/Source/Concert/Public/ConcertMessageData.h
+++ b/Engine/Plugins/Developer/Concert/ConcertMain/Source/Concert/Public/ConcertMessageData.h
@@ -218,10 +218,6 @@ struct FConcertSessionInfo
 	UPROPERTY(VisibleAnywhere, Category="Session Info")
 	FConcertSessionSettings Settings;
 
-	/** The last time the session directory was modified in local time. */
-	UPROPERTY(VisibleAnywhere, Category="Session Info")
-	FDateTime LastModified;
-
 	/** Version information for this session. This is set during creation, and updated each time the session is restored */
 	UPROPERTY()
 	TArray<FConcertSessionVersionInfo> VersionInfos;
@@ -230,10 +226,21 @@ struct FConcertSessionInfo
 	UPROPERTY()
 	EConcertSessionState State = EConcertSessionState::Normal;
 
-	void SetLastModifiedToNow()
-	{
-		LastModified = FDateTime::Now();
-	}
+	FDateTime GetLastModified() const { return FDateTime(LastModifiedTicks); }
+	void SetLastModified(const FDateTime& Value) { LastModifiedTicks = Value.GetTicks(); }
+	void SetLastModifiedToNow() { SetLastModified(FDateTime::Now()); }
+
+private:
+	
+	/**
+	 * The last time the session directory was modified in local time.
+	 * 
+	 * Stored in ticks instead of FDateTime because FDateTime is not serialized properly by FStructSerializer::Serialize and FStructDeserializer::Deserialize.
+	 * These functions are used for packing data when sent across the network. The issue is that in UObject/NoExportTypes.h FDateTime does not expose Ticks as UProperty;
+	 * doing this would be the 'proper' fix but it may cause instability this late into 5.1 dev.
+	 */
+	UPROPERTY(VisibleAnywhere, Category="Session Info")
+	int64 LastModifiedTicks {};
 };
 
 /** Holds filter rules used when migrating session data */
diff --git a/Engine/Plugins/Developer/Concert/ConcertUI/ConcertSharedSlate/Source/ConcertSharedSlate/Private/Session/Browser/SConcertSessionBrowser.cpp b/Engine/Plugins/Developer/Concert/ConcertUI/ConcertSharedSlate/Source/ConcertSharedSlate/Private/Session/Browser/SConcertSessionBrowser.cpp
--- a/Engine/Plugins/Developer/Concert/ConcertUI/ConcertSharedSlate/Source/ConcertSharedSlate/Private/Session/Browser/SConcertSessionBrowser.cpp
+++ b/Engine/Plugins/Developer/Concert/ConcertUI/ConcertSharedSlate/Source/ConcertSharedSlate/Private/Session/Browser/SConcertSessionBrowser.cpp
@@ -177,7 +177,7 @@ static FConcertSessionItem GenerateSessionItem(const SessionType& Session)
 		Session.SessionInfo.Settings.ProjectName,
 		Version.ToString(),
 		Session.ServerInfo.ServerFlags,
-		Session.SessionInfo.LastModified,
+		Session.SessionInfo.GetLastModified(),
 		{}
 	};
 }
