Adding startup delegates to ConcertServer for custom UnrealMultiUserServer apps

This change allows other plugins to be notified and provides them a TWeakPtr when a ConcertServer/ConcertSyncServer or a ConcertServerSession has started

diff --git a/Engine/Plugins/Developer/Concert/ConcertMain/Source/Concert/Private/ConcertServer.cpp b/Engine/Plugins/Developer/Concert/ConcertMain/Source/Concert/Private/ConcertServer.cpp
--- a/Engine/Plugins/Developer/Concert/ConcertMain/Source/Concert/Private/ConcertServer.cpp
+++ b/Engine/Plugins/Developer/Concert/ConcertMain/Source/Concert/Private/ConcertServer.cpp
@@ -280,6 +280,8 @@ void FConcertServer::Startup()
 
 		// Try to mount the default session repository configured (if one is configured) to lock the non-sharable session files away from concurrent processes.
 		MountDefaultSessionRepository(Settings.Get());
+
+		OnConcertServerStartupDelegate.Broadcast();
 	}
 }
 
@@ -679,6 +681,16 @@ bool FConcertServer::DestroySession(const FGuid& SessionId, FText& OutFailureRea
 	return Response.ResponseCode == EConcertResponseCode::Success;
 }
 
+FOnConcertServerSessionStartup& FConcertServer::OnConcertServerSessionStartup()
+{
+	return OnConcertServerSessionStartupDelegate;
+}
+
+FOnConcertServerStartup& FConcertServer::OnConcertServerStartup()
+{
+	return OnConcertServerStartupDelegate;
+}
+
 TArray<FConcertSessionInfo> FConcertServer::GetLiveSessionInfos() const
 {
 	TArray<FConcertSessionInfo> SessionsInfo;
@@ -1569,6 +1581,7 @@ TSharedPtr<IConcertServerSession> FConcertServer::CreateLiveSession(const FConce
 	{
 		LiveSessions.Add(LiveSessionInfo.SessionId, LiveSession);
 		LiveSession->Startup();
+		OnConcertServerSessionStartupDelegate.Broadcast(LiveSession);
 		return LiveSession;
 	}
 
diff --git a/Engine/Plugins/Developer/Concert/ConcertMain/Source/Concert/Private/ConcertServer.h b/Engine/Plugins/Developer/Concert/ConcertMain/Source/Concert/Private/ConcertServer.h
--- a/Engine/Plugins/Developer/Concert/ConcertMain/Source/Concert/Private/ConcertServer.h
+++ b/Engine/Plugins/Developer/Concert/ConcertMain/Source/Concert/Private/ConcertServer.h
@@ -51,6 +51,13 @@ public:
 	virtual bool ExportSession(const FGuid& SessionId, const FConcertSessionFilter& SessionFilter, const FString& DestDir, bool bAnonymizeData, FText& OutFailureReason) override;
 	virtual bool RenameSession(const FGuid& SessionId, const FString& NewName, FText& OutFailureReason) override;
 	virtual bool DestroySession(const FGuid& SessionId, FText& OutFailureReason) override;
+	virtual FOnConcertServerSessionStartup& OnConcertServerSessionStartup() override;
+	virtual FOnConcertServerStartup& OnConcertServerStartup() override;
+
+protected:
+
+	FOnConcertServerSessionStartup OnConcertServerSessionStartupDelegate;
+	FOnConcertServerStartup OnConcertServerStartupDelegate;
 
 private:
 	/** Returns the root dir where the servers keeps the its internally created repositories (When the caller doesn't provide the paths). */
diff --git a/Engine/Plugins/Developer/Concert/ConcertMain/Source/Concert/Public/IConcertServer.h b/Engine/Plugins/Developer/Concert/ConcertMain/Source/Concert/Public/IConcertServer.h
--- a/Engine/Plugins/Developer/Concert/ConcertMain/Source/Concert/Public/IConcertServer.h
+++ b/Engine/Plugins/Developer/Concert/ConcertMain/Source/Concert/Public/IConcertServer.h
@@ -8,6 +8,9 @@
 
 class UConcertServerConfig;
 
+DECLARE_MULTICAST_DELEGATE_OneParam(FOnConcertServerSessionStartup, TWeakPtr<IConcertServerSession>);
+DECLARE_MULTICAST_DELEGATE(FOnConcertServerStartup);
+
 /** Interface for Concert server */
 class IConcertServer
 {
@@ -191,4 +194,14 @@ public:
 	 * @return true if the session was found and destroyed
 	 */
 	virtual bool DestroySession(const FGuid& SessionId, FText& OutFailureReason) = 0;
+
+	/**
+	 * Callback for when a concert server session has been created.
+	 */
+	virtual FOnConcertServerSessionStartup& OnConcertServerSessionStartup() = 0;
+
+	/**
+	 * Callback for when a concert server has been created.
+	 */
+	virtual FOnConcertServerStartup& OnConcertServerStartup() = 0;
 };
diff --git a/Engine/Plugins/Developer/Concert/ConcertSync/ConcertSyncServer/Source/ConcertSyncServer/Private/ConcertSyncServerModule.cpp b/Engine/Plugins/Developer/Concert/ConcertSync/ConcertSyncServer/Source/ConcertSyncServer/Private/ConcertSyncServerModule.cpp
--- a/Engine/Plugins/Developer/Concert/ConcertSync/ConcertSyncServer/Source/ConcertSyncServer/Private/ConcertSyncServerModule.cpp
+++ b/Engine/Plugins/Developer/Concert/ConcertSync/ConcertSyncServer/Source/ConcertSyncServer/Private/ConcertSyncServerModule.cpp
@@ -68,8 +68,19 @@ public:
 
 	virtual TSharedRef<IConcertSyncServer> CreateServer(const FString& InRole, const FConcertSessionFilter& InAutoArchiveSessionFilter) override
 	{
-		return MakeShared<FConcertSyncServer>(InRole, InAutoArchiveSessionFilter);
+		TSharedRef<IConcertSyncServer> ConcertSyncServer = MakeShared<FConcertSyncServer>(InRole, InAutoArchiveSessionFilter);
+		OnConcertSyncServerCreatedDelegate.Broadcast(ConcertSyncServer);
+		return ConcertSyncServer;
 	}
+
+	virtual FOnConcertSyncServerCreated& OnServerCreated() override
+	{
+		return OnConcertSyncServerCreatedDelegate;
+	}
+
+protected:
+
+	FOnConcertSyncServerCreated OnConcertSyncServerCreatedDelegate;
 };
 
 IMPLEMENT_MODULE(FConcertSyncServerModule, ConcertSyncServer);
diff --git a/Engine/Plugins/Developer/Concert/ConcertSync/ConcertSyncServer/Source/ConcertSyncServer/Public/IConcertSyncServerModule.h b/Engine/Plugins/Developer/Concert/ConcertSync/ConcertSyncServer/Source/ConcertSyncServer/Public/IConcertSyncServerModule.h
--- a/Engine/Plugins/Developer/Concert/ConcertSync/ConcertSyncServer/Source/ConcertSyncServer/Public/IConcertSyncServerModule.h
+++ b/Engine/Plugins/Developer/Concert/ConcertSync/ConcertSyncServer/Source/ConcertSyncServer/Public/IConcertSyncServerModule.h
@@ -11,6 +11,8 @@ class UConcertServerConfig;
 
 struct FConcertSessionFilter;
 
+DECLARE_MULTICAST_DELEGATE_OneParam(FOnConcertSyncServerCreated, TWeakPtr<IConcertSyncServer>);
+
 /**
  * Interface for the Concert Sync Server module.
  */
@@ -54,4 +56,9 @@ public:
 	 * @return The server
 	 */
 	virtual TSharedRef<IConcertSyncServer> CreateServer(const FString& InRole, const FConcertSessionFilter& InAutoArchiveSessionFilter) = 0;
+
+	/**
+	 * Callback for when a ConcertSyncServer has been created.
+	 */
+	virtual FOnConcertSyncServerCreated& OnServerCreated() = 0;
 };
