[Concert] Adapting ConcertTransport LocalEndpoint to remove RemoteEndpoints when transport node is forgotten

This change will adapt the ConcertLocalEndpoint in ConcertTransport to remove RemoteEndpoints after a unregister notification from the message bus has been received that a transport node was forgotten. This is needed as otherwise LocalEndpoints have to wait until the timeout removes RemoteEndpoints instead of removing them immediately after the transport node is unregistered.

Tested:
- Connected with UdpMessaging on localhost, joined Multi-User session and changed UdpMessaging settings on the client side
- Connected with QuicMessaging, joined Multi-User session and removed the static server endpoint on the client side

diff --git a/Engine/Plugins/Developer/Concert/ConcertMain/Source/ConcertTransport/Private/ConcertLocalEndpoint.cpp b/Engine/Plugins/Developer/Concert/ConcertMain/Source/ConcertTransport/Private/ConcertLocalEndpoint.cpp
--- a/Engine/Plugins/Developer/Concert/ConcertMain/Source/ConcertTransport/Private/ConcertLocalEndpoint.cpp
+++ b/Engine/Plugins/Developer/Concert/ConcertMain/Source/ConcertTransport/Private/ConcertLocalEndpoint.cpp
@@ -298,6 +298,11 @@ FConcertRemoteEndpointRef FConcertLocalEndpoint::CreateRemoteEndpoint(const FCon
 		FScopeLock RemoteEndpointsLock(&RemoteEndpointsCS);
 		RemoteEndpoints.Add(InEndpointContext.EndpointId, NewRemoteEndpoint);
 	}
+
+	{
+		FScopeLock RemoteAddressesLock(&RemoteAddressesCS);
+		RemoteAddresses.Add(InRemoteAddress);
+	}
 	
 	NewRemoteEndpoint->OnConcertMessageAcknowledgementReceived().AddLambda(
 		[this](const FConcertEndpointContext& Context, const TSharedRef<IConcertMessage>& AckedMessage, const FConcertMessageContext& AckMessageContext)
@@ -326,6 +331,31 @@ FConcertRemoteEndpointPtr FConcertLocalEndpoint::FindRemoteEndpoint(const FMessa
 	return nullptr;
 }
 
+void FConcertLocalEndpoint::RemoveRemoteEndpoint(const FMessageAddress& EndpointAddress)
+{
+	FScopeLock RemoteEndpointsLock(&RemoteEndpointsCS);
+
+	TArray<FConcertRemoteEndpointPtr> Endpoints;
+	RemoteEndpoints.GenerateValueArray(Endpoints);
+
+	for (FConcertRemoteEndpointPtr RemoteEndpoint : Endpoints)
+	{
+		if (RemoteEndpoint->GetAddress() == EndpointAddress)
+		{
+			const FGuid& EndpointId = RemoteEndpoint->GetEndpointContext().EndpointId;
+
+			PendingRemoteEndpointConnectionChangedEvents.Add(MakeTuple(RemoteEndpoint->GetEndpointContext(), EConcertRemoteEndpointConnection::ClosedRemotely));
+			Logger.LogRemoteEndpointClosure(EndpointId, FDateTime::UtcNow());
+			RemoteEndpoints.Remove(EndpointId);
+
+			FScopeLock RemoteAddressesLock(&RemoteAddressesCS);
+			RemoteAddresses.Remove(EndpointAddress);
+
+			break;
+		}
+	}
+}
+
 bool FConcertLocalEndpoint::HandleTick(float DeltaTime)
 {
 	LLM_SCOPE_BYTAG(Concert_ConcertLocalEndpoint);
@@ -587,6 +617,20 @@ void FConcertLocalEndpoint::InternalHandleBusNotification(const FMessageBusNotif
 	{
 		RemoteEndpoint->ForwardBusNotification(Notification.NotificationType);
 	}
+
+	if (Notification.NotificationType == EMessageBusNotification::Unregistered)
+	{
+		{
+			FScopeLock RemoteAddressesLock(&RemoteAddressesCS);
+
+			if (!RemoteAddresses.Contains(Notification.RegistrationAddress))
+			{
+				return;
+			}
+		}
+
+		RemoveRemoteEndpoint(Notification.RegistrationAddress);
+	}
 }
 
 void FConcertLocalEndpoint::ProcessEndpointDiscovery(const FConcertMessageContext& ConcertContext, const FMessageAddress& InRemoteAddress)
@@ -858,6 +902,12 @@ void FConcertLocalEndpoint::TimeoutRemoteEndpoints(const FDateTime& UtcNow)
 				PendingRemoteEndpointConnectionChangedEvents.Add(MakeTuple(RemoteEndpoint->GetEndpointContext(), EConcertRemoteEndpointConnection::TimedOut));
 				Logger.LogRemoteEndpointTimeOut(EndpointId, UtcNow);
 				SendEndpointClosed(RemoteEndpoint.ToSharedRef(), UtcNow);
+
+				{
+					FScopeLock RemoteAddressesLock(&RemoteAddressesCS);
+					RemoteAddresses.Remove(RemoteEndpoint->GetAddress());
+				}
+
 				It.RemoveCurrent();
 				continue;
 			}
diff --git a/Engine/Plugins/Developer/Concert/ConcertMain/Source/ConcertTransport/Private/ConcertLocalEndpoint.h b/Engine/Plugins/Developer/Concert/ConcertMain/Source/ConcertTransport/Private/ConcertLocalEndpoint.h
--- a/Engine/Plugins/Developer/Concert/ConcertMain/Source/ConcertTransport/Private/ConcertLocalEndpoint.h
+++ b/Engine/Plugins/Developer/Concert/ConcertMain/Source/ConcertTransport/Private/ConcertLocalEndpoint.h
@@ -62,6 +62,9 @@ private:
 	/** Find the remote endpoint with the given Message Address */
 	FConcertRemoteEndpointPtr FindRemoteEndpoint(const FMessageAddress& InEndpointId) const;
 
+	/** Remove the remote endpoint with the given Message Address. */
+	void RemoveRemoteEndpoint(const FMessageAddress& EndpointAddress);
+
 	/** Handle the endpoint ticker */
 	bool HandleTick(float DeltaTime);
 
@@ -153,6 +156,10 @@ private:
 	mutable FCriticalSection RemoteEndpointsCS;
 	TMap<FGuid, FConcertRemoteEndpointPtr> RemoteEndpoints;
 
+	/** Array of remote endpoint message addresses that were discovered. */
+	mutable FCriticalSection RemoteAddressesCS;
+	TArray<FMessageAddress> RemoteAddresses;
+
 	/** Incoming message bus messages */
 	TQueue<TSharedPtr<IMessageContext, ESPMode::ThreadSafe>, EQueueMode::Spsc> InboundMessages;
 
