Fix horde issue 243324

diff --git a/Engine/Plugins/Developer/Concert/ConcertMain/Source/Concert/Private/ConcertMessageData.cpp b/Engine/Plugins/Developer/Concert/ConcertMain/Source/Concert/Private/ConcertMessageData.cpp
--- a/Engine/Plugins/Developer/Concert/ConcertMain/Source/Concert/Private/ConcertMessageData.cpp
+++ b/Engine/Plugins/Developer/Concert/ConcertMain/Source/Concert/Private/ConcertMessageData.cpp
@@ -233,8 +233,7 @@ bool TryCompressImpl(const UScriptStruct* InEventType, const void* InEventData,
 		int32 CompressedSize = FCompression::CompressMemoryBound(NamedCompressionAlgo, InOutPayload.PayloadSize, CompressFlags);
 		OutCompressedData.SetNumUninitialized(CompressedSize);
 		SCOPED_CONCERT_TRACE(SerializePayload_CompressMemory);
-		if (FCompression::CompressMemory(NamedCompressionAlgo, OutCompressedData.GetData(), CompressedSize,
-										 InBytes.GetData(), InBytes.Num()), CompressFlags)
+		if (FCompression::CompressMemory(NamedCompressionAlgo, OutCompressedData.GetData(), CompressedSize, InBytes.GetData(), InBytes.Num(), CompressFlags))
 		{
 			OutCompressedData.SetNum(CompressedSize, false);
 			InOutPayload.PayloadBytes.Bytes = MoveTemp(OutCompressedData);
