Disable Oodle compressor for Multi-user there are some cases that cause the compressor to fail.

diff --git a/Engine/Plugins/Developer/Concert/ConcertMain/Source/Concert/Private/ConcertMessageData.cpp b/Engine/Plugins/Developer/Concert/ConcertMain/Source/Concert/Private/ConcertMessageData.cpp
--- a/Engine/Plugins/Developer/Concert/ConcertMain/Source/Concert/Private/ConcertMessageData.cpp
+++ b/Engine/Plugins/Developer/Concert/ConcertMain/Source/Concert/Private/ConcertMessageData.cpp
@@ -172,7 +172,7 @@ namespace UE::Concert::Compression
 {
 
 static TAutoConsoleVariable<int32> CVarCompressionType(
-	TEXT("Concert.SetCompressionType"), 1,
+	TEXT("Concert.SetCompressionType"), 2,
 	TEXT("Specify the type of compression to use when serializing data. A value of 0 means compression is off. A value of 1 = Oodle. All other values = Zlib."));
 
 static TAutoConsoleVariable<int32> CVarCompressionSizeLimit(
@@ -271,11 +271,12 @@ OptionalDecompressBytes DecompressImpl(const FConcertSessionSerializedPayload& I
 		{
 			return OptionalDecompressBytes(MoveTemp(UncompressedData));
 		}
+		else
+		{
+			UE_LOG(LogConcert, Warning, TEXT("Unable to uncompress data for %s!"), *InPayload.PayloadTypeName.ToString());
+		}
 	}
 
-	check(UE::Concert::Compression::DataIsUncompressed(CompressAlgo) &&
-		  InPayload.PayloadBytes.Bytes.Num() == InPayload.PayloadSize);
-
 	return OptionalDecompressBytes{};
 }
 
