Split concert client code from concert server code to allow us to incorporate game tags in the editor client.

diff --git a/Engine/Plugins/Developer/Concert/ConcertApp/DisasterRecoveryClient/Source/DisasterRecoveryClient/DisasterRecoveryClient.Build.cs b/Engine/Plugins/Developer/Concert/ConcertApp/DisasterRecoveryClient/Source/DisasterRecoveryClient/DisasterRecoveryClient.Build.cs
--- a/Engine/Plugins/Developer/Concert/ConcertApp/DisasterRecoveryClient/Source/DisasterRecoveryClient/DisasterRecoveryClient.Build.cs
+++ b/Engine/Plugins/Developer/Concert/ConcertApp/DisasterRecoveryClient/Source/DisasterRecoveryClient/DisasterRecoveryClient.Build.cs
@@ -19,6 +19,7 @@ namespace UnrealBuildTool.Rules
 				new string[]
 				{
 					"Concert",
+					"ConcertClient",
 					"ConcertSyncClient",
 					"ConcertSyncCore",
 					"ConcertSharedSlate",
diff --git a/Engine/Plugins/Developer/Concert/ConcertApp/DisasterRecoveryClient/Source/DisasterRecoveryClient/Private/DisasterRecoveryClientModule.cpp b/Engine/Plugins/Developer/Concert/ConcertApp/DisasterRecoveryClient/Source/DisasterRecoveryClient/Private/DisasterRecoveryClientModule.cpp
--- a/Engine/Plugins/Developer/Concert/ConcertApp/DisasterRecoveryClient/Source/DisasterRecoveryClient/Private/DisasterRecoveryClientModule.cpp
+++ b/Engine/Plugins/Developer/Concert/ConcertApp/DisasterRecoveryClient/Source/DisasterRecoveryClient/Private/DisasterRecoveryClientModule.cpp
@@ -8,6 +8,7 @@
 #include "IConcertClient.h"
 #include "IConcertClientWorkspace.h"
 #include "IConcertSyncClient.h"
+#include "ConcertClientSettings.h"
 #include "ConcertFrontendStyle.h"
 #include "ConcertLocalFileSharingService.h"
 
diff --git a/Engine/Plugins/Developer/Concert/ConcertApp/DisasterRecoveryClient/Source/DisasterRecoveryClient/Private/DisasterRecoverySessionManager.cpp b/Engine/Plugins/Developer/Concert/ConcertApp/DisasterRecoveryClient/Source/DisasterRecoveryClient/Private/DisasterRecoverySessionManager.cpp
--- a/Engine/Plugins/Developer/Concert/ConcertApp/DisasterRecoveryClient/Source/DisasterRecoveryClient/Private/DisasterRecoverySessionManager.cpp
+++ b/Engine/Plugins/Developer/Concert/ConcertApp/DisasterRecoveryClient/Source/DisasterRecoveryClient/Private/DisasterRecoverySessionManager.cpp
@@ -6,6 +6,7 @@
 #include "DisasterRecoverySessionInfo.h"
 #include "IDisasterRecoveryClientModule.h"
 #include "ConcertActivityStream.h"
+#include "ConcertClientSettings.h"
 #include "DisasterRecoverySettings.h"
 #include "GenericPlatform/GenericPlatformCrashContext.h"
 
diff --git a/Engine/Plugins/Developer/Concert/ConcertApp/DisasterRecoveryClient/Source/DisasterRecoveryClient/Private/DisasterRecoveryTasks.h b/Engine/Plugins/Developer/Concert/ConcertApp/DisasterRecoveryClient/Source/DisasterRecoveryClient/Private/DisasterRecoveryTasks.h
--- a/Engine/Plugins/Developer/Concert/ConcertApp/DisasterRecoveryClient/Source/DisasterRecoveryClient/Private/DisasterRecoveryTasks.h
+++ b/Engine/Plugins/Developer/Concert/ConcertApp/DisasterRecoveryClient/Source/DisasterRecoveryClient/Private/DisasterRecoveryTasks.h
@@ -2,7 +2,7 @@
 
 #pragma once
 
-#include "IConcertModule.h"
+#include "IConcertClientModule.h"
 #include "Containers/Ticker.h"
 #include "Misc/Guid.h"
 
diff --git a/Engine/Plugins/Developer/Concert/ConcertApp/MultiUserClient/Source/MultiUserClient/MultiUserClient.Build.cs b/Engine/Plugins/Developer/Concert/ConcertApp/MultiUserClient/Source/MultiUserClient/MultiUserClient.Build.cs
--- a/Engine/Plugins/Developer/Concert/ConcertApp/MultiUserClient/Source/MultiUserClient/MultiUserClient.Build.cs
+++ b/Engine/Plugins/Developer/Concert/ConcertApp/MultiUserClient/Source/MultiUserClient/MultiUserClient.Build.cs
@@ -19,6 +19,7 @@ namespace UnrealBuildTool.Rules
 				{
 					"ApplicationCore",
 					"Concert",
+					"ConcertClient",
 					"ConcertClientSharedSlate",
 					"ConcertSharedSlate",
 					"ConcertSyncClient",
diff --git a/Engine/Plugins/Developer/Concert/ConcertApp/MultiUserServer/Source/MultiUserServer/MultiUserServer.build.cs b/Engine/Plugins/Developer/Concert/ConcertApp/MultiUserServer/Source/MultiUserServer/MultiUserServer.build.cs
--- a/Engine/Plugins/Developer/Concert/ConcertApp/MultiUserServer/Source/MultiUserServer/MultiUserServer.build.cs
+++ b/Engine/Plugins/Developer/Concert/ConcertApp/MultiUserServer/Source/MultiUserServer/MultiUserServer.build.cs
@@ -15,6 +15,7 @@ namespace UnrealBuildTool.Rules
 					"Core",
 					"CoreUObject",
 					"Concert",
+					"ConcertServer",
 					"ConcertSharedSlate",
 					"ConcertSyncCore",
 					"ConcertTransport",
diff --git a/Engine/Plugins/Developer/Concert/ConcertApp/MultiUserServer/Source/MultiUserServer/Private/Widgets/Clients/PackageTransmission/Model/PackageTransmissionModel.cpp b/Engine/Plugins/Developer/Concert/ConcertApp/MultiUserServer/Source/MultiUserServer/Private/Widgets/Clients/PackageTransmission/Model/PackageTransmissionModel.cpp
--- a/Engine/Plugins/Developer/Concert/ConcertApp/MultiUserServer/Source/MultiUserServer/Private/Widgets/Clients/PackageTransmission/Model/PackageTransmissionModel.cpp
+++ b/Engine/Plugins/Developer/Concert/ConcertApp/MultiUserServer/Source/MultiUserServer/Private/Widgets/Clients/PackageTransmission/Model/PackageTransmissionModel.cpp
@@ -6,7 +6,8 @@
 #include "ConcertPackageEvents.h"
 #include "ConcertSyncSessionDatabase.h"
 #include "ConcertUtil.h"
-#include "IConcertModule.h"
+#include "ConcertServerUtil.h"
+#include "IConcertServer.h"
 #include "IConcertSyncServer.h"
 #include "IConcertSession.h"
 #include "Templates/NonNullPointer.h"
diff --git a/Engine/Plugins/Developer/Concert/ConcertApp/MultiUserServer/Source/MultiUserServer/Private/Widgets/Clients/Util/EndpointToUserNameCache.cpp b/Engine/Plugins/Developer/Concert/ConcertApp/MultiUserServer/Source/MultiUserServer/Private/Widgets/Clients/Util/EndpointToUserNameCache.cpp
--- a/Engine/Plugins/Developer/Concert/ConcertApp/MultiUserServer/Source/MultiUserServer/Private/Widgets/Clients/Util/EndpointToUserNameCache.cpp
+++ b/Engine/Plugins/Developer/Concert/ConcertApp/MultiUserServer/Source/MultiUserServer/Private/Widgets/Clients/Util/EndpointToUserNameCache.cpp
@@ -5,6 +5,7 @@
 #include "ConcertMessages.h"
 #include "ConcertServerEvents.h"
 #include "ConcertUtil.h"
+#include "ConcertServerUtil.h"
 #include "IConcertServer.h"
 #include "IConcertSession.h"
 #include "INetworkMessagingExtension.h"
diff --git a/Engine/Plugins/Developer/Concert/ConcertApp/MultiUserServer/Source/MultiUserServer/Private/Widgets/SessionTabs/Archived/ArchivedConcertSessionTab.cpp b/Engine/Plugins/Developer/Concert/ConcertApp/MultiUserServer/Source/MultiUserServer/Private/Widgets/SessionTabs/Archived/ArchivedConcertSessionTab.cpp
--- a/Engine/Plugins/Developer/Concert/ConcertApp/MultiUserServer/Source/MultiUserServer/Private/Widgets/SessionTabs/Archived/ArchivedConcertSessionTab.cpp
+++ b/Engine/Plugins/Developer/Concert/ConcertApp/MultiUserServer/Source/MultiUserServer/Private/Widgets/SessionTabs/Archived/ArchivedConcertSessionTab.cpp
@@ -4,8 +4,12 @@
 
 #include "ArchivedSessionHistoryController.h"
 #include "ConcertLogGlobal.h"
+#include "IConcertServer.h"
 #include "ConcertSyncSessionDatabase.h"
 #include "ConcertSyncSessionTypes.h"
+
+#include "HistoryEdition/ActivityNode.h"
+#include "HistoryEdition/HistoryAnalysis.h"
 #include "HistoryEdition/HistoryEdition.h"
 #include "HistoryEdition/ActivityGraphIDs.h"
 #include "MultiUserServerModule.h"
@@ -17,6 +21,7 @@
 #include "Widgets/SessionTabs/Archived/SConcertArchivedSessionTabView.h"
 #include "Widgets/HistoryEdition/SMuteActivityDependenciesDialog.h"
 
+
 #include "Algo/AllOf.h"
 #include "Dialog/SMessageDialog.h"
 #include "Settings/MultiUserServerUserPreferences.h"
@@ -63,6 +68,82 @@ const FSlateBrush* FArchivedConcertSessionTab::GetTabIconBrush() const
 	return FConcertFrontendStyle::Get()->GetBrush("Concert.ArchivedSession.Icon");
 }
 
+namespace UE::MultiUserServer::Private
+{
+
+FConcertSessionFilter BuildFilterFrom(const TSet<FActivityID>& ToDelete)
+{
+	FConcertSessionFilter Result;
+	Result.ActivityIdsToExclude.Reserve(ToDelete.Num());
+	for (const FActivityID ActivityID : ToDelete)
+	{
+		Result.ActivityIdsToExclude.Add(ActivityID);
+	}
+	return Result;
+}
+
+}
+
+UE::ConcertSyncCore::FOperationErrorResult DeleteActivitiesInArchivedSession(const TSharedRef<IConcertServer>& Server, const FGuid& ArchivedSessionId, const TSet<FActivityID>& ToDelete)
+{
+	using namespace UE::ConcertSyncCore;
+
+	const FString SessionId = ArchivedSessionId.ToString(EGuidFormats::DigitsWithHyphens);
+	UE_LOG(LogConcert, Log, TEXT("Deleting %d activities from session %s..."), ToDelete.Num(), *SessionId);
+	TOptional<FString> ErrorMessage;
+	ON_SCOPE_EXIT
+	{
+		if (ErrorMessage)
+		{
+			UE_LOG(LogConcert, Error, TEXT("Finished deleting activities from %s with error: %s."), *SessionId, **ErrorMessage);
+		}
+		else
+		{
+			UE_LOG(LogConcert, Log, TEXT("Finished deleting activities from %s successfully."), *SessionId);
+		}
+	};
+
+	const TOptional<FConcertSessionInfo> DeletedSessionInfo = Server->GetArchivedSessionInfo(ArchivedSessionId);
+	if (!DeletedSessionInfo)
+	{
+		ErrorMessage = FString::Printf(TEXT("Session ID %s does not resolve to any archived session!"), *ArchivedSessionId.ToString());
+		return FOperationErrorResult::MakeError(FText::FromString(*ErrorMessage));
+	}
+
+	// Restore the session while skipping all to be deleted activities
+	FText FailureReason;
+	const FConcertSessionFilter Filter = UE::MultiUserServer::Private::BuildFilterFrom(ToDelete);
+	// Do not use same name in case user already created a live session from this archived session (would fail in that case) 
+	FConcertSessionInfo OverrideInfo = *DeletedSessionInfo;
+	OverrideInfo.SessionName = DeletedSessionInfo->SessionName + FGuid::NewGuid().ToString();
+	const TSharedPtr<IConcertServerSession> LiveSession = Server->RestoreSession(ArchivedSessionId, OverrideInfo, Filter, FailureReason);
+	if (!LiveSession)
+	{
+		ErrorMessage = FailureReason.ToString();
+		return FOperationErrorResult::MakeError(FailureReason);
+	}
+	ON_SCOPE_EXIT
+	{
+		Server->DestroySession(LiveSession->GetId(), FailureReason);
+	};
+
+	// The archived session must be removed before it can be overwritten
+	if (!Server->DestroySession(ArchivedSessionId, FailureReason))
+	{
+		ErrorMessage = FailureReason.ToString();
+		return FOperationErrorResult::MakeError(FailureReason); 
+	}
+
+	if (!Server->ArchiveSession(LiveSession->GetId(), DeletedSessionInfo->SessionName, Filter, FailureReason, ArchivedSessionId).IsValid())
+	{
+		ErrorMessage = FailureReason.ToString();
+		return FOperationErrorResult::MakeError(FailureReason); 
+	}
+
+	return FOperationErrorResult::MakeSuccess(); 
+}
+
+
 void FArchivedConcertSessionTab::OnRequestDeleteActivities(const TSet<TSharedRef<FConcertSessionActivity>>& ActivitiesToDelete) const
 {
 	using namespace UE::ConcertSyncCore;
diff --git a/Engine/Plugins/Developer/Concert/ConcertMain/ConcertMain.uplugin b/Engine/Plugins/Developer/Concert/ConcertMain/ConcertMain.uplugin
--- a/Engine/Plugins/Developer/Concert/ConcertMain/ConcertMain.uplugin
+++ b/Engine/Plugins/Developer/Concert/ConcertMain/ConcertMain.uplugin
@@ -23,6 +23,18 @@
 			"LoadingPhase": "PreDefault",
 			"ProgramAllowList": [ "UnrealMultiUserServer", "UnrealMultiUserSlateServer", "UnrealRecoverySvc", "CrashReportClientEditor" ]
 		},
+		{
+			"Name": "ConcertClient",
+			"Type": "UncookedOnly",
+			"LoadingPhase": "PreDefault",
+			"ProgramDenyList": [ "UnrealMultiUserServer", "UnrealMultiUserSlateServer", "UnrealRecoverySvc", "CrashReportClientEditor" ]
+		},
+		{
+			"Name": "ConcertServer",
+			"Type": "UncookedOnly",
+			"LoadingPhase": "PreDefault",
+			"ProgramAllowList": [ "UnrealMultiUserServer", "UnrealMultiUserSlateServer", "UnrealRecoverySvc", "CrashReportClientEditor" ]
+		},
 		{
 			"Name": "ConcertTransport",
 			"Type": "UncookedOnly",
diff --git a/Engine/Plugins/Developer/Concert/ConcertMain/Source/Concert/Concert.Build.cs b/Engine/Plugins/Developer/Concert/ConcertMain/Source/Concert/Concert.Build.cs
--- a/Engine/Plugins/Developer/Concert/ConcertMain/Source/Concert/Concert.Build.cs
+++ b/Engine/Plugins/Developer/Concert/ConcertMain/Source/Concert/Concert.Build.cs
@@ -12,23 +12,16 @@ namespace UnrealBuildTool.Rules
 					"Core",
 					"CoreUObject",  // FH: Can we do without at some point?
 					"ConcertTransport",
-					"Serialization",
+					"Serialization"
 				}
 			);
 
 			PublicIncludePathModuleNames.AddRange(
 				new string[]
 				{
-					"ConcertTransport",
+					"ConcertTransport"
 				}
 			);
-
-			//PrivateDependencyModuleNames.AddRange(
-			//	new string[]
-			//	{
-			//		"ConcertTransport",
-			//	}
-			//);
 		}
 	}
 }
diff --git a/Engine/Plugins/Developer/Concert/ConcertMain/Source/Concert/Private/ConcertLogger.h b/Engine/Plugins/Developer/Concert/ConcertMain/Source/Concert/Private/ConcertLogger.h
--- a/Engine/Plugins/Developer/Concert/ConcertMain/Source/Concert/Private/ConcertLogger.h
+++ b/Engine/Plugins/Developer/Concert/ConcertMain/Source/Concert/Private/ConcertLogger.h
@@ -7,7 +7,7 @@
 #include "IConcertEndpoint.h"
 #include "ConcertMessageData.h"
 #include "IConcertTransportLogger.h"
-
+#include "ConcertUtil.h"
 struct FConcertLog;
 struct FConcertMessageContext;
 
@@ -15,8 +15,6 @@ class FConcertLogger : public IConcertTransportLogger
 {
 public:
 
-	using FLogListener = TFunction<void(const FConcertLog&)>;
-	
 	/** Factory function for use with FConcertTransportLoggerFactory */
 	static IConcertTransportLoggerRef CreateLogger(const FConcertEndpointContext& InOwnerContext, FLogListener LogListenerFunc);
 
diff --git a/Engine/Plugins/Developer/Concert/ConcertMain/Source/Concert/Private/ConcertModule.cpp b/Engine/Plugins/Developer/Concert/ConcertMain/Source/Concert/Private/ConcertModule.cpp
--- a/Engine/Plugins/Developer/Concert/ConcertMain/Source/Concert/Private/ConcertModule.cpp
+++ b/Engine/Plugins/Developer/Concert/ConcertMain/Source/Concert/Private/ConcertModule.cpp
@@ -1,40 +1,6 @@
 // Copyright Epic Games, Inc. All Rights Reserved.
 
 #include "IConcertModule.h"
+#include "IConcertTransportModule.h"
 
-#include "ConcertSettings.h"
-#include "ConcertServer.h"
-#include "ConcertClient.h"
-
-/**
- * Implements the Concert module
- */
-class FConcertModule : public IConcertModule
-{
-public:
-
-	virtual void StartupModule() override
-	{
-		EndpointProvider = IConcertTransportModule::Get().CreateEndpointProvider();
-	}
-
-	virtual void ShutdownModule() override
-	{
-		EndpointProvider.Reset();
-	}
-
-	virtual IConcertServerRef CreateServer(const FString& InRole, const FConcertSessionFilter& InAutoArchiveSessionFilter, IConcertServerEventSink* InEventSink) override
-	{
-		return MakeShared<FConcertServer, ESPMode::ThreadSafe>(InRole, InAutoArchiveSessionFilter, InEventSink, EndpointProvider);
-	}
-
-	virtual IConcertClientRef CreateClient(const FString& InRole) override
-	{
-		return MakeShared<FConcertClient, ESPMode::ThreadSafe>(InRole, EndpointProvider);
-	}
-
-private:
-	TSharedPtr<IConcertEndpointProvider> EndpointProvider;
-};
-
-IMPLEMENT_MODULE(FConcertModule, Concert);
+IMPLEMENT_MODULE(FDefaultModuleImpl, Concert);
diff --git a/Engine/Plugins/Developer/Concert/ConcertMain/Source/Concert/Private/ConcertSettings.cpp b/Engine/Plugins/Developer/Concert/ConcertMain/Source/Concert/Private/ConcertSettings.cpp
--- a/Engine/Plugins/Developer/Concert/ConcertMain/Source/Concert/Private/ConcertSettings.cpp
+++ b/Engine/Plugins/Developer/Concert/ConcertMain/Source/Concert/Private/ConcertSettings.cpp
@@ -56,17 +56,5 @@ FText ValidateSessionName(const FString& Name)
 
 } // ConcertSettingsUtils
 
-
-UConcertServerConfig::UConcertServerConfig()
-	: bCleanWorkingDir(false)
-	, NumSessionsToKeep(-1)
-{
-	DefaultVersionInfo.Initialize(false /* bSupportMixedBuildTypes */);
-}
-
-UConcertClientConfig::UConcertClientConfig()
-{
-}
-
 #undef LOCTEXT_NAMESPACE
 
diff --git a/Engine/Plugins/Developer/Concert/ConcertMain/Source/Concert/Private/ConcertUtil.cpp b/Engine/Plugins/Developer/Concert/ConcertMain/Source/Concert/Private/ConcertUtil.cpp
--- a/Engine/Plugins/Developer/Concert/ConcertMain/Source/Concert/Private/ConcertUtil.cpp
+++ b/Engine/Plugins/Developer/Concert/ConcertMain/Source/Concert/Private/ConcertUtil.cpp
@@ -5,8 +5,7 @@
 #include "HAL/FileManager.h"
 #include "ConcertLogGlobal.h"
 #include "ConcertLogger.h"
-#include "ConcertServerSession.h"
-#include "IConcertServer.h"
+#include "IConcertSession.h"
 
 bool ConcertUtil::DeleteDirectoryTree(const TCHAR* InDirectoryToDelete, const TCHAR* InMoveToDirBeforeDelete)
 {
@@ -60,47 +59,7 @@ void ConcertUtil::SetVerboseLogging(bool bInState)
 	FConcertLogger::SetVerboseLogging(bInState);
 }
 
-TArray<FConcertSessionClientInfo> ConcertUtil::GetSessionClients(IConcertServer& Server, const FGuid& SessionId)
+IConcertTransportLoggerRef ConcertUtil::CreateLogger(const FConcertEndpointContext& InOwnerContext, FLogListener LogListenerFunc)
 {
-	TSharedPtr<IConcertServerSession> ServerSession = Server.GetLiveSession(SessionId);
-	if (ServerSession)
-	{
-		return ServerSession->GetSessionClients();
-	}
-	return TArray<FConcertSessionClientInfo>();
+	return FConcertLogger::CreateLogger(InOwnerContext, MoveTemp(LogListenerFunc));
 }
-
-namespace ConcertUtil::Private
-{
-	static TOptional<TPair<FConcertSessionClientInfo, TSharedPtr<IConcertServerSession>>> FindByClient(IConcertServer& Server, const FGuid& ClientEndpointId)
-	{
-		for (const TSharedPtr<IConcertServerSession>& ServerSession : Server.GetLiveSessions())
-		{
-			for (const FConcertSessionClientInfo& ClientInfo : ServerSession->GetSessionClients())
-			{
-				if (ClientEndpointId == ClientInfo.ClientEndpointId)
-				{
-					return {{ ClientInfo, ServerSession }};
-				}
-			}
-		}
-
-		return {};
-	}
-}
-
-TOptional<FConcertSessionClientInfo> ConcertUtil::GetConnectedClientInfo(IConcertServer& Server, const FGuid& ClientEndpointId)
-{
-	const TOptional<TPair<FConcertSessionClientInfo, TSharedPtr<IConcertServerSession>>> Result = Private::FindByClient(Server, ClientEndpointId);
-	return Result
-		? Result->Key
-		: TOptional<FConcertSessionClientInfo>();
-}
-
-TSharedPtr<IConcertServerSession> ConcertUtil::GetLiveSessionClientConnectedTo(IConcertServer& Server, const FGuid& ClientEndpointId)
-{
-	const TOptional<TPair<FConcertSessionClientInfo, TSharedPtr<IConcertServerSession>>> Result = Private::FindByClient(Server, ClientEndpointId);
-	return Result
-		? Result->Value
-		: TSharedPtr<IConcertServerSession>();
-}
\ No newline at end of file
diff --git a/Engine/Plugins/Developer/Concert/ConcertMain/Source/Concert/Public/ConcertSession.h b/Engine/Plugins/Developer/Concert/ConcertMain/Source/Concert/Public/ConcertSession.h
new file mode 100644
--- /dev/null
+++ b/Engine/Plugins/Developer/Concert/ConcertMain/Source/Concert/Public/ConcertSession.h
@@ -0,0 +1,83 @@
+// Copyright Epic Games, Inc. All Rights Reserved.
+
+#pragma once
+
+#include "CoreMinimal.h"
+#include "IConcertSession.h"
+#include "ConcertMessages.h"
+
+/**
+ * Common implementation for Concert Client and Server sessions (@see IConcertSession for the API description).
+ * @note	This doesn't inherit from any session interface, but does implement some of their API with a "Common" prefix on the function names.
+ *			Interface implementations can also inherit from this common impl and then call the "Common" functions from the interface overrides.
+ */
+class CONCERT_API FConcertSessionCommonImpl
+{
+public:
+	explicit FConcertSessionCommonImpl(const FConcertSessionInfo& InSessionInfo);
+
+	void CommonStartup();
+	void CommonShutdown();
+
+	const FGuid& CommonGetId() const
+	{
+		return SessionInfo.SessionId;
+	}
+
+	const FString& CommonGetName() const
+	{
+		return SessionInfo.SessionName;
+	}
+
+	void CommonSetName(const FString& NewName)
+	{
+		SessionInfo.SessionName = NewName;
+	}
+
+	const FConcertSessionInfo& CommonGetSessionInfo() const
+	{
+		return SessionInfo;
+	}
+
+	TArray<FGuid> CommonGetSessionClientEndpointIds() const;
+	TArray<FConcertSessionClientInfo> CommonGetSessionClients() const;
+	bool CommonFindSessionClient(const FGuid& EndpointId, FConcertSessionClientInfo& OutSessionClientInfo) const;
+
+	FConcertScratchpadRef CommonGetScratchpad() const;
+	FConcertScratchpadPtr CommonGetClientScratchpad(const FGuid& ClientEndpointId) const;
+
+	FDelegateHandle CommonRegisterCustomEventHandler(const FName& EventMessageType, const TSharedRef<IConcertSessionCustomEventHandler>& Handler);
+	void CommonUnregisterCustomEventHandler(const FName& EventMessageType, const FDelegateHandle EventHandle);
+	void CommonUnregisterCustomEventHandler(const FName& EventMessageType, const void* EventHandler);
+	void CommonClearCustomEventHandler(const FName& EventMessageType);
+	void CommonHandleCustomEvent(const FConcertMessageContext& Context);
+
+	void CommonRegisterCustomRequestHandler(const FName& RequestMessageType, const TSharedRef<IConcertSessionCustomRequestHandler>& Handler);
+	void CommonUnregisterCustomRequestHandler(const FName& RequestMessageType);
+	TFuture<FConcertSession_CustomResponse> CommonHandleCustomRequest(const FConcertMessageContext& Context);
+
+	static bool CommonBuildCustomEvent(const UScriptStruct* EventType, const void* EventData, const FGuid& SourceEndpointId, const TArray<FGuid>& DestinationEndpointIds, FConcertSession_CustomEvent& OutCustomEvent);
+	static bool CommonBuildCustomRequest(const UScriptStruct* RequestType, const void* RequestData, const FGuid& SourceEndpointId, const FGuid& DestinationEndpointId, FConcertSession_CustomRequest& OutCustomRequest);
+	static void CommonHandleCustomResponse(const FConcertSession_CustomResponse& Response, const TSharedRef<IConcertSessionCustomResponseHandler>& Handler);
+
+protected:
+	/** Information about this session */
+	FConcertSessionInfo SessionInfo;
+
+	/** The scratchpad for this session */
+	FConcertScratchpadPtr Scratchpad;
+
+	/** Map of clients connected to this session (excluding us if we're a client session) */
+	struct FSessionClient
+	{
+		FConcertSessionClientInfo ClientInfo;
+		FConcertScratchpadPtr Scratchpad;
+	};
+	TMap<FGuid, FSessionClient> SessionClients;
+
+	/** Map of custom event handlers for this session */
+	TMap<FName, TArray<TSharedPtr<IConcertSessionCustomEventHandler>>> CustomEventHandlers;
+
+	/** Map of custom request handlers for this session */
+	TMap<FName, TSharedPtr<IConcertSessionCustomRequestHandler>> CustomRequestHandlers;
+};
diff --git a/Engine/Plugins/Developer/Concert/ConcertMain/Source/Concert/Public/ConcertSettings.h b/Engine/Plugins/Developer/Concert/ConcertMain/Source/Concert/Public/ConcertSettings.h
--- a/Engine/Plugins/Developer/Concert/ConcertMain/Source/Concert/Public/ConcertSettings.h
+++ b/Engine/Plugins/Developer/Concert/ConcertMain/Source/Concert/Public/ConcertSettings.h
@@ -9,19 +9,19 @@
 #include "Misc/EngineVersion.h"
 #include "ConcertTransportSettings.h"
 #include "ConcertVersion.h"
+
 #include "ConcertSettings.generated.h"
 
 namespace ConcertSettingsUtils
 {
 
-/** Returns an error messages if the user display name is invalid, otherwise, returns an empty text. */
+	/** Returns an error messages if the user display name is invalid, otherwise, returns an empty text. */
 FText CONCERT_API ValidateDisplayName(const FString& Name);
 
 /** Returns an error messages if the specified session name is invalid, otherwise, returns an empty text. */
 FText CONCERT_API ValidateSessionName(const FString& Name);
 
 }
-
 USTRUCT()
 struct FConcertSessionSettings
 {
@@ -79,322 +79,3 @@ struct FConcertSessionSettings
 
 	// TODO: private session, password, etc etc,
 };
-
-USTRUCT()
-struct FConcertServerSettings
-{
-	GENERATED_BODY()
-
-	/** The server will allow client to join potentially incompatible sessions */
-	UPROPERTY(config, EditAnywhere, AdvancedDisplay, Category="Server Settings")
-	bool bIgnoreSessionSettingsRestriction = false;
-
-	/** The timespan at which session updates are processed. */
-	UPROPERTY(config, EditAnywhere, DisplayName="Session Tick Frequency", AdvancedDisplay, Category="Server Settings", meta=(ForceUnits=s))
-	int32 SessionTickFrequencySeconds = 1;
-};
-
-UCLASS(config=Engine)
-class CONCERT_API UConcertServerConfig : public UObject
-{
-	GENERATED_BODY()
-public:
-	UConcertServerConfig();
-
-	/**
-	 * If true, instruct the server to auto-archive sessions that were left in the working directory because the server did not exit properly rather than
-	 * restoring them as 'live' (the default).
-	 */
-	UPROPERTY(config)
-	bool bAutoArchiveOnReboot = false;
-
-	/**
-	 * If true, instruct the server to auto-archive live sessions on shutdown.
-	 */
-	UPROPERTY(config)
-	bool bAutoArchiveOnShutdown = true;
-
-	/**
-	 * Clean server sessions working directory when booting
-	 * Can be specified on the server cmd with `-CONCERTCLEAN`
-	 */
-	UPROPERTY(config, EditAnywhere, Category="Server Settings")
-	bool bCleanWorkingDir;
-
-	/**
-	 * Number of archived sessions to keep when booting, or <0 to keep all archived sessions
-	 */
-	UPROPERTY(config, EditAnywhere, Category="Server Settings")
-	int32 NumSessionsToKeep;
-
-	/** 
-	 * Name of the server, or empty to use the default name.
-	 * Can be specified on the server cmd with `-CONCERTSERVER=`
-	 */
-	UPROPERTY(config, EditAnywhere, Category="Server Settings")
-	FString ServerName;
-
-	/** 
-	 * Name of the default session created on the server.
-	 * Can be specified on the server cmd with `-CONCERTSESSION=`
-	 */
-	UPROPERTY(config, EditAnywhere, Category="Session Settings")
-	FString DefaultSessionName;
-
-	/** 
-	 * A set of keys identifying the clients that can discover and access the server. If empty, the server can be discovered and used by any clients.
-	 */
-	UPROPERTY(config)
-	TSet<FString> AuthorizedClientKeys;
-
-	/**
-	 * Name of the default session to restore on the server.
-	 * Set the name of the desired save to restore its content in your session.
-	 * Leave this blank if you want to create an empty session.
-	 * Can be specified on the editor cmd with `-CONCERTSESSIONTORESTORE=`.
-	 */
-	UPROPERTY(config, EditAnywhere, Category="Session Settings")
-	FString DefaultSessionToRestore;
-
-	/** 
-	 * The version string for the default server created.
-	 * Can be specified on the server cmd with `-CONCERTVERSION=`
-	 */
-	UPROPERTY()
-	FConcertSessionVersionInfo DefaultVersionInfo;
-
-	/** Default server session settings */
-	UPROPERTY(config, EditAnywhere, Category="Session Settings")
-	FConcertSessionSettings DefaultSessionSettings;
-
-	/** Server & server session settings */
-	UPROPERTY(config, EditAnywhere, Category="Server Settings", meta=(ShowOnlyInnerProperties))
-	FConcertServerSettings ServerSettings;
-
-	/** Endpoint settings passed down to endpoints on creation */
-	UPROPERTY(config, EditAnywhere, AdvancedDisplay, Category="Endpoint Settings", meta=(ShowOnlyInnerProperties))
-	FConcertEndpointSettings EndpointSettings;
-
-	/** The default directory where the server keeps the live session files. Can be specified on the server command line with `-CONCERTWORKINGDIR=`*/
-	UPROPERTY(config)
-	FString WorkingDir;
-
-	/** The default directory where the server keeps the archived session files. Can be specified on the server command line with `-CONCERTSAVEDDIR=`*/
-	UPROPERTY(config)
-	FString ArchiveDir;
-
-	/** The root directory where the server creates new session repositories (unless the client request specifies its own root). If empty or invalid, the server will use a default. */
-	UPROPERTY(config)
-	FString SessionRepositoryRootDir;
-
-	/** If neither of WorkingDir and ArchiveDir are set, determine whether the server should mount a standard default session repository where new session will be created. */
-	UPROPERTY(config)
-	bool bMountDefaultSessionRepository = true;
-};
-
-USTRUCT()
-struct FConcertClientSettings
-{
-	GENERATED_BODY()
-
-	FConcertClientSettings()
-		: DisplayName()
-		, AvatarColor(1.0f, 1.0f, 1.0f, 1.0f)
-		, DesktopAvatarActorClass(TEXT("/ConcertSyncClient/DesktopPresence.DesktopPresence_C"))
-		, VRAvatarActorClass(TEXT("/ConcertSyncClient/VRPresence.VRPresence_C"))
-		, ServerPort(0)
-		, DiscoveryTimeoutSeconds(5)
-		, SessionTickFrequencySeconds(1)
-		, LatencyCompensationMs(0)
-	{}
-
-	/** 
-	 * The display name to use when in a session. 
-	 * Can be specified on the editor cmd with `-CONCERTDISPLAYNAME=`.
-	 */
-	UPROPERTY(config, EditAnywhere, Category="Client Settings")
-	FString DisplayName;
-
-	/** The color used for the presence avatar in a session. */
-	UPROPERTY(config, EditAnywhere, Category="Client Settings")
-	FLinearColor AvatarColor;
-
-	/** The desktop representation of this editor's user to other connected users */
-	UPROPERTY(config, EditAnywhere, NoClear, Category = "Client Settings", meta = (MetaClass = "/Script/ConcertSyncClient.ConcertClientDesktopPresenceActor"))
-	FSoftClassPath DesktopAvatarActorClass;
-
-	/** The VR representation of this editor's user to other connected users */
-	UPROPERTY(config, EditAnywhere, NoClear, Category = "Client Settings", meta = (MetaClass = "/Script/ConcertSyncClient.ConcertClientVRPresenceActor", DisplayName = "VR Avatar Actor Class"))
-	FSoftClassPath VRAvatarActorClass;
-
-	/** The port to use to reach the server with static endpoints when launched through the editor. This port will be used over the unicast endpoint port in the UDP Messagging settings if non 0 when transferring the editor settings to the launched server. */
-	UPROPERTY(config, EditAnywhere, AdvancedDisplay, Category = "Client Settings")
-	uint16 ServerPort;
-
-	/** The timespan at which discovered Multi-User server are considered stale if they haven't answered back */
-	UPROPERTY(config, EditAnywhere, DisplayName="Discovery Timeout", AdvancedDisplay, Category="Client Settings", meta=(ForceUnits=s))
-	int32 DiscoveryTimeoutSeconds;
-
-	/** The timespan at which session updates are processed. */
-	UPROPERTY(config, EditAnywhere, DisplayName="Session Tick Frequency", AdvancedDisplay, Category="Client Settings", meta=(ForceUnits=s))
-	int32 SessionTickFrequencySeconds;
-
-	/** Amount of latency compensation to apply to time-synchronization sensitive interactions */
-	UPROPERTY(config, EditAnywhere, DisplayName="Latency Compensation", AdvancedDisplay, Category="Client Settings", meta=(ForceUnits=ms))
-	float LatencyCompensationMs;
-
-	/**
-	 * When level editor changes are made, reflect those changes to the game equivalent property.
-	 * This settings can be specified on the editor cmd with `-CONCERTREFLECTVISIBILITY=`.
-	 */
-	UPROPERTY(config, EditAnywhere, DisplayName="Reflect Level Visibility to Game", AdvancedDisplay, Category="Client Settings")
-	bool bReflectLevelEditorInGame = false;
-
-	/**
-	 * Enable extended version support when using Multi-user with precompiled and source builds.  When using Unreal Game
-	 * Sync, it is possible to have the same engine CL but different engine version due to content changes.	 This setting
-	 * enables reading engine version CL from the Build.version file produced by UGS to determine engine version
-	 * information when joining a session.	This only applies when you intend to mix precompiled with source builds.
-	*/
-	UPROPERTY(config, EditAnywhere, DisplayName="Support Mixed Build Types", AdvancedDisplay, Category="Client Settings")
-	bool bSupportMixedBuildTypes = false;
-
-	/** Array of tags that can be used for grouping and categorizing. */
-	UPROPERTY(config, EditAnywhere, AdvancedDisplay, Category = "Client Settings")
-	TArray<FName> Tags;
-
-	/** A key used to identify the clients during server discovery. If the server was configured to restrict access, the client key must be know of the server. Can be left empty. */
-	UPROPERTY(config)
-	FString ClientAuthenticationKey;
-};
-
-UENUM()
-enum class EConcertSourceValidationMode : uint8
-{
-	/** Source control validation will fail on any changes when connecting to a Multi-User Session. */
-	Hard = 0,
-	/** 
-	 * Source control validation will warn and prompt on any changes when connecting to a Multi-User session. 
-	 * In Memory changes will be hot-reloaded.
-	 * Source control changes aren't affected but will be stashed/shelved in the future.
-	 */
-	Soft,
-	/** Soft validation mode with auto proceed on prompts. */
-	SoftAutoProceed
-};
-
-USTRUCT()
-struct FConcertSourceControlSettings
-{
-	GENERATED_BODY()
-
-	FConcertSourceControlSettings()
-		: ValidationMode(EConcertSourceValidationMode::Soft)
-	{}
-
-	UPROPERTY(config, EditAnywhere, Category="Revision Control Settings", meta=(Keywords = "Source Control"))
-	EConcertSourceValidationMode ValidationMode;
-};
-
-UENUM()
-enum class EConcertServerType
-{
-	Console,
-	Slate
-};
-
-UCLASS(config=Engine)
-class CONCERT_API UConcertClientConfig : public UObject
-{
-	GENERATED_BODY()
-public:
-	UConcertClientConfig();
-
-	/*
-	 * Mark this setting object as editor only.
-	 * This so soft object path reference made by this setting object won't be automatically grabbed by the cooker.
-	 * @see UPackage::Save, FSoftObjectPathThreadContext::GetSerializationOptions, FSoftObjectPath::ImportTextItem
-	 * @todo: cooker should have a better way to filter editor only objects for 'unsolicited' references.
-	 */
-	virtual bool IsEditorOnly() const
-	{
-		return true;
-	}
-
-	/**
-	 * True if this client should be "headless"? (ie, not display any UI).
-	 */
-	UPROPERTY(config)
-	bool bIsHeadless;
-
-	/**
-	 * True if the Multi-User module should install shortcut button and its drop-down menu in the level editor toolbar.
-	 */
-	UPROPERTY(config, EditAnywhere, Category="Client Settings", Meta=(ConfigRestartRequired=true, DisplayName="Enable Multi-User Toolbar Button"))
-	bool bInstallEditorToolbarButton;
-
-	/** 
-	 * Automatically connect or create default session on default server.
-	 * Can be specified on the editor cmd with `-CONCERTAUTOCONNECT` or `-CONCERTAUTOCONNECT=<true/false>`.
-	 */
-	UPROPERTY(config, EditAnywhere, Category="Client Settings")
-	bool bAutoConnect;
-
-	/** 
-	 * If auto-connect is on, retry connecting to the default server/session until it succeeds or the user cancels.
-	 * Can be specified on the editor cmd with `-CONCERTRETRYAUTOCONNECTONERROR` or `-CONCERTRETRYAUTOCONNECTONERROR=<true/false>`.
-	 */
-	UPROPERTY(config, EditAnywhere, Category="Client Settings")
-	bool bRetryAutoConnectOnError = false;
-
-	/**
-	 * Determines which server executable with be launched:
-	 *		Console -> UnrealMultiUserServer.exe
-	 *		Slate	-> UnrealMultiUserSlateServer.exe
-	 */
-	UPROPERTY(config, EditAnywhere, Category="Client Settings")
-	EConcertServerType ServerType = EConcertServerType::Slate;
-
-	/** 
-	 * Default server url (just a name for now) to look for on auto or default connect. 
- 	 * Can be specified on the editor cmd with `-CONCERTSERVER=`.
-	 */
-	UPROPERTY(config, EditAnywhere, Category="Client Settings")
-	FString DefaultServerURL;
-
-	/** 
-	 * Default session name to look for on auto connect or default connect.
-	 * Can be specified on the editor cmd with `-CONCERTSESSION=`.
-	 */
-	UPROPERTY(config, EditAnywhere, Category="Client Settings")
-	FString DefaultSessionName;
-
-	/**
-	 * If this client create the default session, should the session restore a saved session.
-	 * Set the name of the desired save to restore its content in your session.
-	 * Leave this blank if you want to create an empty session.
-	 * Can be specified on the editor cmd with `-CONCERTSESSIONTORESTORE=`.
-	 */
-	UPROPERTY(config, EditAnywhere, Category = "Client Settings")
-	FString DefaultSessionToRestore;
-
-	/**
-	 * If this client create the default session, should the session data be saved when it's deleted.
-	 * Set the name desired for the save and the session data will be moved in that save when the session is deleted
-	 * Leave this blank if you don't want to save the session data.
-	 * Can be specified on the editor cmd with `-CONCERTSAVESESSIONAS=`.
-	*/
-	UPROPERTY(config, EditAnywhere, Category = "Client Settings")
-	FString DefaultSaveSessionAs;
-
-	/** Client & client session settings */
-	UPROPERTY(config, EditAnywhere, Category="Client Settings", meta=(ShowOnlyInnerProperties))
-	FConcertClientSettings ClientSettings;
-
-	UPROPERTY(config, EditAnywhere, Category = "Revision Control Settings", meta=(ShowOnlyInnerProperties, Keywords = "Source Control"))
-	FConcertSourceControlSettings SourceControlSettings;
-
-	/** Endpoint settings passed down to endpoints on creation */
-	UPROPERTY(config, EditAnywhere, AdvancedDisplay, Category="Endpoint Settings", meta=(ShowOnlyInnerProperties))
-	FConcertEndpointSettings EndpointSettings;
-};
diff --git a/Engine/Plugins/Developer/Concert/ConcertMain/Source/Concert/Public/ConcertTransportEvents.h b/Engine/Plugins/Developer/Concert/ConcertMain/Source/Concert/Public/ConcertTransportEvents.h
--- a/Engine/Plugins/Developer/Concert/ConcertMain/Source/Concert/Public/ConcertTransportEvents.h
+++ b/Engine/Plugins/Developer/Concert/ConcertMain/Source/Concert/Public/ConcertTransportEvents.h
@@ -4,7 +4,6 @@
 
 #include "CoreMinimal.h"
 #include "ConcertMessageData.h"
-#include "IConcertClient.h"
 #include "ConcertTransportEvents.generated.h"
 
 class IConcertClient;
diff --git a/Engine/Plugins/Developer/Concert/ConcertMain/Source/Concert/Public/ConcertUtil.h b/Engine/Plugins/Developer/Concert/ConcertMain/Source/Concert/Public/ConcertUtil.h
--- a/Engine/Plugins/Developer/Concert/ConcertMain/Source/Concert/Public/ConcertUtil.h
+++ b/Engine/Plugins/Developer/Concert/ConcertMain/Source/Concert/Public/ConcertUtil.h
@@ -3,10 +3,13 @@
 #pragma once
 
 #include "CoreMinimal.h"
+#include "IConcertTransportLogger.h"
 
 class IConcertServerSession;
 class IConcertServer;
 struct FConcertSessionClientInfo;
+struct FConcertLog;
+using FLogListener = TFunction<void(const FConcertLog&)>;
 
 namespace ConcertUtil
 {
@@ -20,27 +23,7 @@ CONCERT_API bool Copy(FArchive& DstAr, FArchive& SrcAr, int64 Size);
 /** Turn on verbose logging for all loggers (including console loggers). */
 CONCERT_API void SetVerboseLogging(bool bInState);
 
-/**
- * Get the list of clients for a session
- * @param Server The server to look on
- * @param SessionId The session ID
- * @return A list of clients connected to the session
- */
-CONCERT_API TArray<FConcertSessionClientInfo> GetSessionClients(IConcertServer& Server, const FGuid& SessionId);
-
-/**
- * Gets a connected client's info.
- * @param Server The server to look on
- * @param ClientEndpointId The client's endpoint ID
- * @return The client's info if found (it is not found if the client is not connected to any session)
- */
-CONCERT_API TOptional<FConcertSessionClientInfo> GetConnectedClientInfo(IConcertServer& Server, const FGuid& ClientEndpointId);
-
-/**
- * Retrieves the server session object from the given connected client.
- * @param Server The server to look on
- * @param ClientEndpointId The client's endpoint ID
- * @return The session the client is connected to if any
- */
-CONCERT_API TSharedPtr<IConcertServerSession> GetLiveSessionClientConnectedTo(IConcertServer& Server, const FGuid& ClientEndpointId);
+/** Factory function for use with FConcertTransportLoggerFactory */
+CONCERT_API IConcertTransportLoggerRef CreateLogger(const FConcertEndpointContext& InOwnerContext, FLogListener LogListenerFunc);
+
 } // namespace ConcertUtil
diff --git a/Engine/Plugins/Developer/Concert/ConcertMain/Source/Concert/Public/IConcertModule.h b/Engine/Plugins/Developer/Concert/ConcertMain/Source/Concert/Public/IConcertModule.h
--- a/Engine/Plugins/Developer/Concert/ConcertMain/Source/Concert/Public/IConcertModule.h
+++ b/Engine/Plugins/Developer/Concert/ConcertMain/Source/Concert/Public/IConcertModule.h
@@ -6,17 +6,6 @@
 #include "Modules/ModuleInterface.h"
 #include "Modules/ModuleManager.h"
 
-class IConcertServer;
-class IConcertClient;
-class IConcertServerEventSink;
-
-struct FConcertSessionFilter;
-
-typedef TSharedPtr<IConcertServer, ESPMode::ThreadSafe> IConcertServerPtr;
-typedef TSharedRef<IConcertServer, ESPMode::ThreadSafe> IConcertServerRef;
-typedef TSharedPtr<IConcertClient, ESPMode::ThreadSafe> IConcertClientPtr;
-typedef TSharedRef<IConcertClient, ESPMode::ThreadSafe> IConcertClientRef;
-
 /**
  * Interface for the Main Concert module.
  */
@@ -45,20 +34,4 @@ public:
 		static const FName ModuleName = "Concert";
 		return FModuleManager::Get().IsModuleLoaded(ModuleName);
 	}
-
-	/**
-	 * Create a server that will perform a certain role (eg, MultiUser, DisasterRecovery, etc)
-	 * @param InRole The role to create
-	 * @param InAutoArchiveSessionFilter The session filter to apply when auto-archiving sessions
-	 * @param InEventSink Sink functions for events that the server can emit
-	 * @return The server
-	 */
-	virtual IConcertServerRef CreateServer(const FString& InRole, const FConcertSessionFilter& InAutoArchiveSessionFilter, IConcertServerEventSink* InEventSink) = 0;
-
-	/**
-	 * Create a client that will perform a certain role (eg, MultiUser, DisasterRecovery, etc)
-	 * @param InRole The role to create
-	 * @return The client
-	 */
-	virtual IConcertClientRef CreateClient(const FString& InRole) = 0;
 };
diff --git a/Engine/Plugins/Developer/Concert/ConcertSync/ConcertSyncClient/Source/ConcertSyncClient/ConcertSyncClient.Build.cs b/Engine/Plugins/Developer/Concert/ConcertSync/ConcertSyncClient/Source/ConcertSyncClient/ConcertSyncClient.Build.cs
--- a/Engine/Plugins/Developer/Concert/ConcertSync/ConcertSyncClient/Source/ConcertSyncClient/ConcertSyncClient.Build.cs
+++ b/Engine/Plugins/Developer/Concert/ConcertSync/ConcertSyncClient/Source/ConcertSyncClient/ConcertSyncClient.Build.cs
@@ -13,7 +13,8 @@ namespace UnrealBuildTool.Rules
 					"CoreUObject",
 					"Engine",
 					"Concert",
-					"ConcertSyncCore",
+					"ConcertClient",
+					"ConcertSyncCore"
 				}
 			);
 
diff --git a/Engine/Plugins/Developer/Concert/ConcertSync/ConcertSyncClient/Source/ConcertSyncClient/Private/ConcertActivityStream.cpp b/Engine/Plugins/Developer/Concert/ConcertSync/ConcertSyncClient/Source/ConcertSyncClient/Private/ConcertActivityStream.cpp
--- a/Engine/Plugins/Developer/Concert/ConcertSync/ConcertSyncClient/Source/ConcertSyncClient/Private/ConcertActivityStream.cpp
+++ b/Engine/Plugins/Developer/Concert/ConcertSync/ConcertSyncClient/Source/ConcertSyncClient/Private/ConcertActivityStream.cpp
@@ -5,6 +5,7 @@
 #include "IConcertClientWorkspace.h"
 #include "IConcertClient.h"
 #include "IConcertModule.h"
+#include "IConcertClientModule.h"
 
 #define LOCTEXT_NAMESPACE "ConcertActivityStream"
 
diff --git a/Engine/Plugins/Developer/Concert/ConcertSync/ConcertSyncClient/Source/ConcertSyncClient/Private/ConcertSyncClient.cpp b/Engine/Plugins/Developer/Concert/ConcertSync/ConcertSyncClient/Source/ConcertSyncClient/Private/ConcertSyncClient.cpp
--- a/Engine/Plugins/Developer/Concert/ConcertSync/ConcertSyncClient/Source/ConcertSyncClient/Private/ConcertSyncClient.cpp
+++ b/Engine/Plugins/Developer/Concert/ConcertSync/ConcertSyncClient/Source/ConcertSyncClient/Private/ConcertSyncClient.cpp
@@ -5,6 +5,7 @@
 #include "IConcertModule.h"
 #include "IConcertClient.h"
 #include "IConcertSession.h"
+#include "ConcertClientSettings.h"
 #include "ConcertSyncClientLiveSession.h"
 #include "ConcertClientWorkspace.h"
 #include "ConcertClientSequencerManager.h"
@@ -15,7 +16,7 @@ LLM_DECLARE_TAG(Concert_ConcertSyncClient);
 #define LOCTEXT_NAMESPACE "ConcertSyncClient"
 
 FConcertSyncClient::FConcertSyncClient(const FString& InRole, IConcertClientPackageBridge* InPackageBridge, IConcertClientTransactionBridge* InTransactionBridge)
-	: ConcertClient(IConcertModule::Get().CreateClient(InRole))
+	: ConcertClient(IConcertClientModule::Get().CreateClient(InRole))
 	, SessionFlags(EConcertSyncSessionFlags::None)
 	, PackageBridge(InPackageBridge)
 	, TransactionBridge(InTransactionBridge)
diff --git a/Engine/Plugins/Developer/Concert/ConcertSync/ConcertSyncClient/Source/ConcertSyncClient/Public/IConcertSyncClient.h b/Engine/Plugins/Developer/Concert/ConcertSync/ConcertSyncClient/Source/ConcertSyncClient/Public/IConcertSyncClient.h
--- a/Engine/Plugins/Developer/Concert/ConcertSync/ConcertSyncClient/Source/ConcertSyncClient/Public/IConcertSyncClient.h
+++ b/Engine/Plugins/Developer/Concert/ConcertSync/ConcertSyncClient/Source/ConcertSyncClient/Public/IConcertSyncClient.h
@@ -5,6 +5,7 @@
 #include "CoreMinimal.h"
 #include "IConcertClientTransactionBridge.h"
 #include "IConcertModule.h"
+#include "IConcertClientModule.h"
 #include "ConcertActionDefinition.h"
 #include "ConcertSyncSessionFlags.h"
 
diff --git a/Engine/Plugins/Developer/Concert/ConcertSync/ConcertSyncCore/Source/ConcertSyncCore/Private/HistoryEdition/HistoryEdition.cpp b/Engine/Plugins/Developer/Concert/ConcertSync/ConcertSyncCore/Source/ConcertSyncCore/Private/HistoryEdition/HistoryEdition.cpp
--- a/Engine/Plugins/Developer/Concert/ConcertSync/ConcertSyncCore/Source/ConcertSyncCore/Private/HistoryEdition/HistoryEdition.cpp
+++ b/Engine/Plugins/Developer/Concert/ConcertSync/ConcertSyncCore/Source/ConcertSyncCore/Private/HistoryEdition/HistoryEdition.cpp
@@ -4,14 +4,9 @@
 
 #include "ConcertLogGlobal.h"
 #include "ConcertSyncSessionDatabase.h"
-#include "IConcertServer.h"
-
-#include "Misc/ScopeExit.h"
 
 namespace UE::ConcertSyncCore
 {
-	static FConcertSessionFilter BuildFilterFrom(const TSet<FActivityID>& ToDelete);
-
 	TSet<FActivityID> CombineRequirements(const FHistoryAnalysisResult& ToDelete)
 	{
 		TSet<FActivityID> Result;
@@ -25,73 +20,5 @@ namespace UE::ConcertSyncCore
 		}
 		return Result;
 	}
-	
-	FOperationErrorResult DeleteActivitiesInArchivedSession(const TSharedRef<IConcertServer>& Server, const FGuid& ArchivedSessionId, const TSet<FActivityID>& ToDelete)
-	{
-		const FString SessionId = ArchivedSessionId.ToString(EGuidFormats::DigitsWithHyphens);
-		UE_LOG(LogConcert, Log, TEXT("Deleting %d activities from session %s..."), ToDelete.Num(), *SessionId);
-		TOptional<FString> ErrorMessage;
-		ON_SCOPE_EXIT
-		{
-			if (ErrorMessage)
-			{
-				UE_LOG(LogConcert, Error, TEXT("Finished deleting activities from %s with error: %s."), *SessionId, **ErrorMessage);
-			}
-			else
-			{
-				UE_LOG(LogConcert, Log, TEXT("Finished deleting activities from %s successfully."), *SessionId);
-			}
-		};
-		
-		const TOptional<FConcertSessionInfo> DeletedSessionInfo = Server->GetArchivedSessionInfo(ArchivedSessionId);
-		if (!DeletedSessionInfo)
-		{
-			ErrorMessage = FString::Printf(TEXT("Session ID %s does not resolve to any archived session!"), *ArchivedSessionId.ToString());
-			return FOperationErrorResult::MakeError(FText::FromString(*ErrorMessage));
-		}
-
-		// Restore the session while skipping all to be deleted activities
-		FText FailureReason;
-		const FConcertSessionFilter Filter = BuildFilterFrom(ToDelete);
-		// Do not use same name in case user already created a live session from this archived session (would fail in that case) 
-		FConcertSessionInfo OverrideInfo = *DeletedSessionInfo;
-		OverrideInfo.SessionName = DeletedSessionInfo->SessionName + FGuid::NewGuid().ToString();
-		const TSharedPtr<IConcertServerSession> LiveSession = Server->RestoreSession(ArchivedSessionId, OverrideInfo, Filter, FailureReason);
-		if (!LiveSession)
-		{
-			ErrorMessage = FailureReason.ToString();
-			return FOperationErrorResult::MakeError(FailureReason);
-		}
-		ON_SCOPE_EXIT
-		{
-			Server->DestroySession(LiveSession->GetId(), FailureReason);
-		};
-
-		// The archived session must be removed before it can be overwritten
-		if (!Server->DestroySession(ArchivedSessionId, FailureReason))
-		{
-			ErrorMessage = FailureReason.ToString();
-			return FOperationErrorResult::MakeError(FailureReason); 
-		}
-
-		if (!Server->ArchiveSession(LiveSession->GetId(), DeletedSessionInfo->SessionName, Filter, FailureReason, ArchivedSessionId).IsValid())
-		{
-			ErrorMessage = FailureReason.ToString();
-			return FOperationErrorResult::MakeError(FailureReason); 
-		}
-		
-		return FOperationErrorResult::MakeSuccess(); 
-	}
-
-	static FConcertSessionFilter BuildFilterFrom(const TSet<FActivityID>& ToDelete)
-	{
-		FConcertSessionFilter Result;
-		Result.ActivityIdsToExclude.Reserve(ToDelete.Num());
-		for (const FActivityID ActivityID : ToDelete)
-		{
-			Result.ActivityIdsToExclude.Add(ActivityID);
-		}
-		return Result;
-	}
 }
 
diff --git a/Engine/Plugins/Developer/Concert/ConcertSync/ConcertSyncCore/Source/ConcertSyncCore/Public/HistoryEdition/HistoryEdition.h b/Engine/Plugins/Developer/Concert/ConcertSync/ConcertSyncCore/Source/ConcertSyncCore/Public/HistoryEdition/HistoryEdition.h
--- a/Engine/Plugins/Developer/Concert/ConcertSync/ConcertSyncCore/Source/ConcertSyncCore/Public/HistoryEdition/HistoryEdition.h
+++ b/Engine/Plugins/Developer/Concert/ConcertSync/ConcertSyncCore/Source/ConcertSyncCore/Public/HistoryEdition/HistoryEdition.h
@@ -27,14 +27,4 @@ namespace UE::ConcertSyncCore
 
 	/** Utility functions that converts FHistoryDeletionRequirements into a single TSet. */
 	CONCERTSYNCCORE_API TSet<FActivityID> CombineRequirements(const FHistoryAnalysisResult& ToDelete);
-
-	/**
-	 * Deletes the given sessions in ToDelete from ArchivedSessionDatabase.
-	 *
-	 * There is no direct functionality for removing activities in the database; the operation is as follows:
-	 *	1. Restore the archived session with a session filter
-	 *	2. Delete the archived session
-	 *	3. Archive the live session created in step 1 with the old guid.
-	 */
-	CONCERTSYNCCORE_API FOperationErrorResult DeleteActivitiesInArchivedSession(const TSharedRef<IConcertServer>& Server, const FGuid& ArchivedSessionId, const TSet<FActivityID>& ToDelete);
 }
diff --git a/Engine/Plugins/Developer/Concert/ConcertSync/ConcertSyncServer/Source/ConcertSyncServer/ConcertSyncServer.Build.cs b/Engine/Plugins/Developer/Concert/ConcertSync/ConcertSyncServer/Source/ConcertSyncServer/ConcertSyncServer.Build.cs
--- a/Engine/Plugins/Developer/Concert/ConcertSync/ConcertSyncServer/Source/ConcertSyncServer/ConcertSyncServer.Build.cs
+++ b/Engine/Plugins/Developer/Concert/ConcertSync/ConcertSyncServer/Source/ConcertSyncServer/ConcertSyncServer.Build.cs
@@ -19,6 +19,7 @@ namespace UnrealBuildTool.Rules
 					"CoreUObject",
 					"Concert",
 					"ConcertSyncCore",
+					"ConcertServer"
 				}
 			);
 
diff --git a/Engine/Plugins/Developer/Concert/ConcertSync/ConcertSyncServer/Source/ConcertSyncServer/Private/ConcertSyncServer.cpp b/Engine/Plugins/Developer/Concert/ConcertSync/ConcertSyncServer/Source/ConcertSyncServer/Private/ConcertSyncServer.cpp
--- a/Engine/Plugins/Developer/Concert/ConcertSync/ConcertSyncServer/Source/ConcertSyncServer/Private/ConcertSyncServer.cpp
+++ b/Engine/Plugins/Developer/Concert/ConcertSync/ConcertSyncServer/Source/ConcertSyncServer/Private/ConcertSyncServer.cpp
@@ -227,7 +227,7 @@ bool MigrateSessionData(const FConcertSyncSessionDatabase& InSourceDatabase, con
 } // namespace ConcertSyncServerUtils
 
 FConcertSyncServer::FConcertSyncServer(const FString& InRole, const FConcertSessionFilter& InAutoArchiveSessionFilter)
-	: ConcertServer(IConcertModule::Get().CreateServer(InRole, InAutoArchiveSessionFilter, this))
+	: ConcertServer(IConcertServerModule::Get().CreateServer(InRole, InAutoArchiveSessionFilter, this))
 	, SessionFlags(EConcertSyncSessionFlags::None)
 {
 }
diff --git a/Engine/Plugins/Developer/Concert/ConcertSync/ConcertSyncServer/Source/ConcertSyncServer/Private/ConcertSyncServerModule.cpp b/Engine/Plugins/Developer/Concert/ConcertSync/ConcertSyncServer/Source/ConcertSyncServer/Private/ConcertSyncServerModule.cpp
--- a/Engine/Plugins/Developer/Concert/ConcertSync/ConcertSyncServer/Source/ConcertSyncServer/Private/ConcertSyncServerModule.cpp
+++ b/Engine/Plugins/Developer/Concert/ConcertSync/ConcertSyncServer/Source/ConcertSyncServer/Private/ConcertSyncServerModule.cpp
@@ -3,6 +3,7 @@
 #include "IConcertSyncServerModule.h"
 #include "ConcertSyncServer.h"
 #include "ConcertSettings.h"
+#include "ConcertServerSettings.h"
 #include "Logging/LogVerbosity.h"
 #include "Misc/EngineVersion.h"
 #include "ConcertLogGlobal.h"
diff --git a/Engine/Plugins/Developer/Concert/ConcertSync/ConcertSyncServer/Source/ConcertSyncServer/Public/ConcertSyncServerLoop.inl b/Engine/Plugins/Developer/Concert/ConcertSync/ConcertSyncServer/Source/ConcertSyncServer/Public/ConcertSyncServerLoop.inl
--- a/Engine/Plugins/Developer/Concert/ConcertSync/ConcertSyncServer/Source/ConcertSyncServer/Public/ConcertSyncServerLoop.inl
+++ b/Engine/Plugins/Developer/Concert/ConcertSync/ConcertSyncServer/Source/ConcertSyncServer/Public/ConcertSyncServerLoop.inl
@@ -4,6 +4,7 @@
 
 #include "ConcertSyncServerLoop.h"
 #include "ConcertSettings.h"
+#include "ConcertServerSettings.h"
 #include "IConcertServer.h"
 #include "IConcertSyncServer.h"
 #include "IConcertSyncServerModule.h"
diff --git a/Engine/Plugins/Developer/Concert/ConcertSync/ConcertSyncServer/Source/ConcertSyncServer/Public/IConcertSyncServer.h b/Engine/Plugins/Developer/Concert/ConcertSync/ConcertSyncServer/Source/ConcertSyncServer/Public/IConcertSyncServer.h
--- a/Engine/Plugins/Developer/Concert/ConcertSync/ConcertSyncServer/Source/ConcertSyncServer/Public/IConcertSyncServer.h
+++ b/Engine/Plugins/Developer/Concert/ConcertSync/ConcertSyncServer/Source/ConcertSyncServer/Public/IConcertSyncServer.h
@@ -4,6 +4,7 @@
 
 #include "CoreMinimal.h"
 #include "IConcertModule.h"
+#include "IConcertServerModule.h"
 #include "ConcertSyncSessionFlags.h"
 #include "Templates/NonNullPointer.h"
 
diff --git a/Engine/Plugins/Experimental/VirtualProductionUtilities/Source/VPUtilities/Private/VPUtilitiesModule.cpp b/Engine/Plugins/Experimental/VirtualProductionUtilities/Source/VPUtilities/Private/VPUtilitiesModule.cpp
--- a/Engine/Plugins/Experimental/VirtualProductionUtilities/Source/VPUtilities/Private/VPUtilitiesModule.cpp
+++ b/Engine/Plugins/Experimental/VirtualProductionUtilities/Source/VPUtilities/Private/VPUtilitiesModule.cpp
@@ -1,10 +1,56 @@
 // Copyright Epic Games, Inc. All Rights Reserved.
 
 #include "VPUtilitiesModule.h"
+#include "GameplayTagContainer.h"
+#include "Modules/ModuleInterface.h"
 #include "Modules/ModuleManager.h"
+#include "IConcertClient.h"
 
+#include "VPRolesSubsystem.h"
+#include "UnrealEngine.h"
 
 DEFINE_LOG_CATEGORY(LogVPUtilities);
 
+namespace UE::VPUtilities::Private
+{
 
-IMPLEMENT_MODULE(FDefaultModuleImpl, VPUtilities)
+bool EvaluateRole(const FGameplayTagContainer& Container)
+{
+	UVirtualProductionRolesSubsystem* VPRolesSubsystem = GEngine->GetEngineSubsystem<UVirtualProductionRolesSubsystem>();
+	if (!VPRolesSubsystem)
+	{
+		return false;
+	}
+
+	for (auto It = Container.CreateConstIterator(); It; ++It)
+	{
+		if (VPRolesSubsystem->HasActiveRole(*It))
+		{
+			return true;
+		}
+	}
+	return false;
+}
+
+}
+/**
+ * Implements the Concert Client module
+ */
+class FVPUtilitiesModule : public IModuleInterface
+{
+public:
+
+	virtual void StartupModule() override
+	{
+		UE::ConcertClient::VPRoleEvaluator().BindStatic(
+			&UE::VPUtilities::Private::EvaluateRole );
+	}
+
+	virtual void ShutdownModule() override
+	{
+	}
+
+private:
+};
+
+IMPLEMENT_MODULE(FVPUtilitiesModule, VPUtilities)
diff --git a/Engine/Plugins/Experimental/VirtualProductionUtilities/Source/VPUtilities/VPUtilities.Build.cs b/Engine/Plugins/Experimental/VirtualProductionUtilities/Source/VPUtilities/VPUtilities.Build.cs
--- a/Engine/Plugins/Experimental/VirtualProductionUtilities/Source/VPUtilities/VPUtilities.Build.cs
+++ b/Engine/Plugins/Experimental/VirtualProductionUtilities/Source/VPUtilities/VPUtilities.Build.cs
@@ -32,6 +32,8 @@ public class VPUtilities : ModuleRules
 				"RHI",
 				"Slate",
 				"SlateCore",
+				"ConcertClient",
+				"VPRoles"
 			}
 		);
 
diff --git a/Engine/Source/Programs/UnrealMultiUserServer/UnrealMultiUserServer.Build.cs b/Engine/Source/Programs/UnrealMultiUserServer/UnrealMultiUserServer.Build.cs
--- a/Engine/Source/Programs/UnrealMultiUserServer/UnrealMultiUserServer.Build.cs
+++ b/Engine/Source/Programs/UnrealMultiUserServer/UnrealMultiUserServer.Build.cs
@@ -12,6 +12,7 @@ public class UnrealMultiUserServer : ModuleRules
 				"Core",
 				"CoreUObject",
 				"Concert",
+				"ConcertServer",
 				"TraceInsights",
 				"ApplicationCore",					// for LaunchEngineLoop.cpp dependency
 				"Projects",							// for LaunchEngineLoop.cpp dependency
