Implement support for rewriting Concert transaction data against a different identifier table

This can be used to merge multiple transactions (using different identifier tables) together, either using a single combined identifier table (for portability), or using the in-process FName table (for non-portability, but less memory).

#rnx

diff --git a/Engine/Plugins/Developer/Concert/ConcertMain/Source/ConcertTransport/Private/IdentifierTable/ConcertIdentifierTable.cpp b/Engine/Plugins/Developer/Concert/ConcertMain/Source/ConcertTransport/Private/IdentifierTable/ConcertIdentifierTable.cpp
--- a/Engine/Plugins/Developer/Concert/ConcertMain/Source/ConcertTransport/Private/IdentifierTable/ConcertIdentifierTable.cpp
+++ b/Engine/Plugins/Developer/Concert/ConcertMain/Source/ConcertTransport/Private/IdentifierTable/ConcertIdentifierTable.cpp
@@ -2,11 +2,18 @@
 
 #include "IdentifierTable/ConcertIdentifierTable.h"
 
+FConcertLocalIdentifierTable* FConcertLocalIdentifierTable::ForceFNameTableIndex = (FConcertLocalIdentifierTable*)INDEX_NONE;
+
 FConcertLocalIdentifierTable::FConcertLocalIdentifierTable(const FConcertLocalIdentifierState& InState)
 {
 	SetState(InState);
 }
 
+bool FConcertLocalIdentifierTable::HasMappings() const
+{
+	return MappedNames.Num() > 0;
+}
+
 int32 FConcertLocalIdentifierTable::MapName(const FName& InName)
 {
 	if (const int32* ExistingIndex = NameToMappedIndex.Find(InName))
diff --git a/Engine/Plugins/Developer/Concert/ConcertMain/Source/ConcertTransport/Private/IdentifierTable/ConcertTransportArchives.cpp b/Engine/Plugins/Developer/Concert/ConcertMain/Source/ConcertTransport/Private/IdentifierTable/ConcertTransportArchives.cpp
--- a/Engine/Plugins/Developer/Concert/ConcertMain/Source/ConcertTransport/Private/IdentifierTable/ConcertTransportArchives.cpp
+++ b/Engine/Plugins/Developer/Concert/ConcertMain/Source/ConcertTransport/Private/IdentifierTable/ConcertTransportArchives.cpp
@@ -3,6 +3,9 @@
 #include "IdentifierTable/ConcertTransportArchives.h"
 #include "IdentifierTable/ConcertIdentifierTable.h"
 
+namespace UE::Concert::Private::ConcertIdentifierArchiveUtil
+{
+
 enum class EConcertIdentifierSource : uint8
 {
 	/** Plain string value (no suffix) */
@@ -10,97 +13,75 @@ enum class EConcertIdentifierSource : uint8
 	/** Hardcoded FName index value (see MAX_NETWORKED_HARDCODED_NAME) */
 	HardcodedIndex,
 	/** Local identifier table index value (see FConcertLocalIdentifierTable) */
-	LocalIdentifierTableIndex,
+	LocalIdentifierTableIndex_PackedInt,
+	/** Local identifier table index value (see FConcertLocalIdentifierTable) */
+	LocalIdentifierTableIndex_FixedSizeInt32,
+	/** FName table index value (non-portable!) */
+	FNameTableIndex_FixedSizeUInt32,
 };
 
-FConcertIdentifierWriter::FConcertIdentifierWriter(FConcertLocalIdentifierTable* InLocalIdentifierTable, TArray<uint8>& InBytes, bool bIsPersistent)
-	: FMemoryWriter(InBytes, bIsPersistent)
-	, LocalIdentifierTable(InLocalIdentifierTable)
+void WriteName(FArchive& Ar, const FName& Name, FConcertLocalIdentifierTable* LocalIdentifierTable)
 {
-}
-
-FArchive& FConcertIdentifierWriter::operator<<(FName& Name)
-{
-	auto SerializeConcertIdentifierSource = [this](EConcertIdentifierSource InSource)
+	auto SerializeConcertIdentifierSource = [&Ar](EConcertIdentifierSource InSource)
 	{
-		Serialize(&InSource, sizeof(EConcertIdentifierSource));
+		Ar.Serialize(&InSource, sizeof(EConcertIdentifierSource));
 	};
 
-	auto SerializeIndexValue = [this](int32 InIndex)
+	auto SerializeIndexValue = [&Ar](int32 InIndex)
 	{
 		check(InIndex >= 0);
 		uint32 UnsignedIndex = (uint32)InIndex;
-		SerializeIntPacked(UnsignedIndex);
+		Ar.SerializeIntPacked(UnsignedIndex);
 	};
 
-	const EName* Ename = Name.ToEName();
-	int32 IdentifierTableIndex = INDEX_NONE;
-	if (Ename && ShouldReplicateAsInteger(*Ename, Name))
+	if (const EName* Ename = Name.ToEName(); Ename && ShouldReplicateAsInteger(*Ename, Name))
 	{
 		SerializeConcertIdentifierSource(EConcertIdentifierSource::HardcodedIndex);
 		SerializeIndexValue((int32)*Ename);
 	}
+	else if (LocalIdentifierTable == FConcertLocalIdentifierTable::ForceFNameTableIndex)
+	{
+		SerializeConcertIdentifierSource(EConcertIdentifierSource::FNameTableIndex_FixedSizeUInt32);
+		uint32 FNameTableIndex = Name.GetDisplayIndex().ToUnstableInt();
+		Ar << FNameTableIndex; // Note: Don't serialize as a packed int so that the data size remains consistent in case this data gets rewritten (see FConcertIdentifierRewriter)
+	}
 	else if (LocalIdentifierTable)
 	{
-		SerializeConcertIdentifierSource(EConcertIdentifierSource::LocalIdentifierTableIndex);
-		IdentifierTableIndex = LocalIdentifierTable->MapName(Name);
-		SerializeIndexValue(IdentifierTableIndex);
+		SerializeConcertIdentifierSource(EConcertIdentifierSource::LocalIdentifierTableIndex_FixedSizeInt32);
+		int32 IdentifierTableIndex = LocalIdentifierTable->MapName(Name);
+		Ar << IdentifierTableIndex; // Note: Don't serialize as a packed int so that the data size remains consistent in case this data gets rewritten (see FConcertIdentifierRewriter)
 	}
 	else
 	{
 		SerializeConcertIdentifierSource(EConcertIdentifierSource::PlainString);
 		FString PlainString = Name.GetPlainNameString();
-		*this << PlainString;
+		Ar << PlainString;
 	}
 	
 	int32 NameNumber = Name.GetNumber();
 	SerializeIndexValue(NameNumber);
-
-	return *this;
-}
-
-FArchive& FConcertIdentifierWriter::operator<<(FSoftObjectPath& AssetPtr)
-{
-	FName ObjPath = *AssetPtr.ToString();
-	*this << ObjPath;
-	return *this;
 }
 
-FString FConcertIdentifierWriter::GetArchiveName() const
+void ReadName(FArchive& Ar, FName& Name, const FConcertLocalIdentifierTable* LocalIdentifierTable)
 {
-	return TEXT("FConcertIdentifierWriter");
-}
+	checkf(LocalIdentifierTable != FConcertLocalIdentifierTable::ForceFNameTableIndex, TEXT("FConcertLocalIdentifierTable::ForceFNameTableIndex can only be used when writing!"));
 
-
-FConcertIdentifierReader::FConcertIdentifierReader(const FConcertLocalIdentifierTable* InLocalIdentifierTable, const TArray<uint8>& InBytes, bool bIsPersistent)
-	: FMemoryReader(InBytes, bIsPersistent)
-	, LocalIdentifierTable(InLocalIdentifierTable)
-{
-}
-
-FArchive& FConcertIdentifierReader::operator<<(FName& Name)
-{
-	if (GetError())
-	{
-		return *this;
-	}
-
-	auto SerializeIndexValue = [this]() -> int32
+	auto SerializeIndexValue = [&Ar]() -> int32
 	{
 		uint32 UnsignedIndex = 0;
-		SerializeIntPacked(UnsignedIndex);
+		Ar.SerializeIntPacked(UnsignedIndex);
 		return (int32)UnsignedIndex;
 	};
 
 	EConcertIdentifierSource Source;
-	Serialize(&Source, sizeof(EConcertIdentifierSource));
+	Ar.Serialize(&Source, sizeof(EConcertIdentifierSource));
 
 	switch (Source)
 	{
 	case EConcertIdentifierSource::PlainString:
 		{
 			FString PlainString;
-			*this << PlainString;
+			Ar << PlainString;
 			Name = FName(*PlainString, NAME_NO_NUMBER_INTERNAL, /*bSplitName*/false);
 		}
 		break;
@@ -112,17 +93,37 @@ FArchive& FConcertIdentifierReader::operator<<(FName& Name)
 		}
 		break;
 
-	case EConcertIdentifierSource::LocalIdentifierTableIndex:
+	case EConcertIdentifierSource::LocalIdentifierTableIndex_PackedInt:
 		{
 			const int32 IdentifierTableIndex = SerializeIndexValue();
 			if (!LocalIdentifierTable || !LocalIdentifierTable->UnmapName(IdentifierTableIndex, Name))
 			{
-				SetError();
-				return *this;
+				Ar.SetError();
+				return;
+			}
+		}
+		break;
+
+	case EConcertIdentifierSource::LocalIdentifierTableIndex_FixedSizeInt32:
+		{
+			int32 IdentifierTableIndex = INDEX_NONE;
+			Ar << IdentifierTableIndex;
+			if (!LocalIdentifierTable || !LocalIdentifierTable->UnmapName(IdentifierTableIndex, Name))
+			{
+				Ar.SetError();
+				return;
 			}
 		}
 		break;
 
+	case EConcertIdentifierSource::FNameTableIndex_FixedSizeUInt32:
+		{
+			uint32 FNameTableIndex = 0;
+			Ar << FNameTableIndex;
+			Name = FName::CreateFromDisplayId(FNameEntryId::FromUnstableInt(FNameTableIndex), 0);
+		}
+		break;
+
 	default:
 		checkf(false, TEXT("Unknown EConcertIdentifierSource!"));
 		break;
@@ -130,14 +131,73 @@ FArchive& FConcertIdentifierReader::operator<<(FName& Name)
 
 	const int32 NameNumber = SerializeIndexValue();
 	Name.SetNumber(NameNumber);
+}
+
+void WriteSoftObjectPath(FArchive& Ar, const FSoftObjectPath& AssetPtr)
+{
+	// TODO: Serialize via FSoftObjectPath::SerializePath to avoid converting all object paths to FName
+	FNameBuilder ObjPathStr;
+	AssetPtr.ToString(ObjPathStr);
+
+	FName ObjPath = FName(ObjPathStr);
+	Ar << ObjPath;
+}
+
+void ReadSoftObjectPath(FArchive& Ar, FSoftObjectPath& AssetPtr)
+{
+	FName ObjPath;
+	Ar << ObjPath;
+
+	FNameBuilder ObjPathStr(ObjPath);
+	AssetPtr.SetPath(ObjPathStr.ToView());
+}
+
+} // namespace UE::Concert::Private::ConcertIdentifierArchiveUtil
+
+FConcertIdentifierWriter::FConcertIdentifierWriter(FConcertLocalIdentifierTable* InLocalIdentifierTable, TArray<uint8>& InBytes, bool bIsPersistent)
+	: FMemoryWriter(InBytes, bIsPersistent)
+	, LocalIdentifierTable(InLocalIdentifierTable)
+{
+}
+
+FArchive& FConcertIdentifierWriter::operator<<(FName& Name)
+{
+	UE::Concert::Private::ConcertIdentifierArchiveUtil::WriteName(*this, Name, LocalIdentifierTable);
+	return *this;
+}
+
+FArchive& FConcertIdentifierWriter::operator<<(FSoftObjectPath& AssetPtr)
+{
+	UE::Concert::Private::ConcertIdentifierArchiveUtil::WriteSoftObjectPath(*this, AssetPtr);
+	return *this;
+}
+
+FString FConcertIdentifierWriter::GetArchiveName() const
+{
+	return TEXT("FConcertIdentifierWriter");
+}
+
+
+FConcertIdentifierReader::FConcertIdentifierReader(const FConcertLocalIdentifierTable* InLocalIdentifierTable, const TArray<uint8>& InBytes, bool bIsPersistent)
+	: FMemoryReader(InBytes, bIsPersistent)
+	, LocalIdentifierTable(InLocalIdentifierTable)
+{
+}
 
+FArchive& FConcertIdentifierReader::operator<<(FName& Name)
+{
+	if (GetError())
+	{
+		return *this;
+	}
+
+	UE::Concert::Private::ConcertIdentifierArchiveUtil::ReadName(*this, Name, LocalIdentifierTable);
 	return *this;
 }
 
 FArchive& FConcertIdentifierReader::operator<<(FSoftObjectPath& AssetPtr)
 {
-	FName ObjPath = *AssetPtr.ToString();
-	*this << ObjPath;
+	UE::Concert::Private::ConcertIdentifierArchiveUtil::ReadSoftObjectPath(*this, AssetPtr);
 	return *this;
 }
 
@@ -146,3 +206,83 @@ FString FConcertIdentifierReader::GetArchiveName() const
 	return TEXT("FConcertIdentifierReader");
 }
 
+
+FConcertIdentifierRewriter::FConcertIdentifierRewriter(const FConcertLocalIdentifierTable* InLocalIdentifierTable, FConcertLocalIdentifierTable* InRewriteIdentifierTable, TArray<uint8>& InBytes, bool bIsPersistent)
+	: LocalIdentifierTable(InLocalIdentifierTable)
+	, RewriteIdentifierTable(InRewriteIdentifierTable)
+	, Bytes(InBytes)
+{
+	checkf(LocalIdentifierTable && RewriteIdentifierTable, TEXT("FConcertIdentifierRewriter can only rewrite data via an identifier table, as the rewritten data must be the same size!"));
+	SetIsLoading(true);
+	SetIsPersistent(bIsPersistent);
+}
+
+FArchive& FConcertIdentifierRewriter::operator<<(FName& Name)
+{
+	if (GetError())
+	{
+		return *this;
+	}
+
+	const int32 OffsetBeforeNameRead = Offset;
+	UE::Concert::Private::ConcertIdentifierArchiveUtil::ReadName(*this, Name, LocalIdentifierTable);
+	const int32 OffsetAfterNameRead = Offset;
+
+	// Write the name to the scratch buffer using the rewrite identifier table
+	{
+		ScratchBytes.Reset();
+		FMemoryWriter ScratchWriter(ScratchBytes, IsPersistent());
+		UE::Concert::Private::ConcertIdentifierArchiveUtil::WriteName(ScratchWriter, Name, RewriteIdentifierTable);
+	}
+
+	// Patch the scratch buffer into the real buffer
+	// The data must be the same size to avoid serialization issues with tagged data
+	{
+		const int64 OldSerializedSize = OffsetAfterNameRead - OffsetBeforeNameRead;
+		const int64 NewSerializedSize = ScratchBytes.Num();
+		if (ensureAlwaysMsgf(OldSerializedSize == NewSerializedSize, TEXT("FConcertIdentifierRewriter: Rewritten data must be the same size to avoid issues with tagged property serialization! (old: %d, new: %d)"), OldSerializedSize, NewSerializedSize))
+		{
+			FMemory::Memcpy(Bytes.GetData() + OffsetBeforeNameRead, ScratchBytes.GetData(), NewSerializedSize);
+		}
+		else
+		{
+			SetError();
+			return *this;
+		}
+	}
+
+	return *this;
+}
+
+FArchive& FConcertIdentifierRewriter::operator<<(FSoftObjectPath& AssetPtr)
+{
+	UE::Concert::Private::ConcertIdentifierArchiveUtil::WriteSoftObjectPath(*this, AssetPtr);
+	return *this;
+}
+
+int64 FConcertIdentifierRewriter::TotalSize()
+{
+	return (int64)Bytes.Num();
+}
+
+void FConcertIdentifierRewriter::Serialize(void* Data, int64 Num)
+{
+	if (Num && !IsError())
+	{
+		// Only serialize if we have the requested amount of data
+		if (Offset + Num <= TotalSize())
+		{
+			FMemory::Memcpy(Data, Bytes.GetData() + Offset, Num);
+			Offset += Num;
+		}
+		else
+		{
+			SetError();
+		}
+	}
+}
+
+FString FConcertIdentifierRewriter::GetArchiveName() const
+{
+	return TEXT("FConcertIdentifierRewriter");
+}
diff --git a/Engine/Plugins/Developer/Concert/ConcertMain/Source/ConcertTransport/Public/IdentifierTable/ConcertIdentifierTable.h b/Engine/Plugins/Developer/Concert/ConcertMain/Source/ConcertTransport/Public/IdentifierTable/ConcertIdentifierTable.h
--- a/Engine/Plugins/Developer/Concert/ConcertMain/Source/ConcertTransport/Public/IdentifierTable/ConcertIdentifierTable.h
+++ b/Engine/Plugins/Developer/Concert/ConcertMain/Source/ConcertTransport/Public/IdentifierTable/ConcertIdentifierTable.h
@@ -30,6 +30,16 @@ public:
 	FConcertLocalIdentifierTable() = default;
 	explicit FConcertLocalIdentifierTable(const FConcertLocalIdentifierState& InState);
 
+	/**
+	 * Constant used to force names to serialize as their FName table index value (non-portable!)
+	 */
+	static FConcertLocalIdentifierTable* ForceFNameTableIndex;
+
+	/**
+	 * Does this table have any mappings in it?
+	 */
+	bool HasMappings() const;
+
 	/**
 	 * Map the given name to its identifier index.
 	 */
diff --git a/Engine/Plugins/Developer/Concert/ConcertMain/Source/ConcertTransport/Public/IdentifierTable/ConcertTransportArchives.h b/Engine/Plugins/Developer/Concert/ConcertMain/Source/ConcertTransport/Public/IdentifierTable/ConcertTransportArchives.h
--- a/Engine/Plugins/Developer/Concert/ConcertMain/Source/ConcertTransport/Public/IdentifierTable/ConcertTransportArchives.h
+++ b/Engine/Plugins/Developer/Concert/ConcertMain/Source/ConcertTransport/Public/IdentifierTable/ConcertTransportArchives.h
@@ -26,7 +26,7 @@ private:
 	FConcertLocalIdentifierTable* LocalIdentifierTable;
 };
 
-/** Archive for reading identifiers (currently names) in a way that avoids duplication by caching them against their internal key, which can then be mapped over the network  */
+/** Archive for reading identifiers (currently names) in a way that avoids duplication by caching them against their internal key, which can then be mapped over the network */
 class CONCERTTRANSPORT_API FConcertIdentifierReader : public FMemoryReader
 {
 public:
@@ -43,3 +43,26 @@ public:
 private:
 	const FConcertLocalIdentifierTable* LocalIdentifierTable;
 };
+
+/** Archive for rewriting identifiers (currently names) so that they belong to a different identifier table */
+class CONCERTTRANSPORT_API FConcertIdentifierRewriter : public FMemoryArchive
+{
+public:
+	FConcertIdentifierRewriter(const FConcertLocalIdentifierTable* InLocalIdentifierTable, FConcertLocalIdentifierTable* InRewriteIdentifierTable, TArray<uint8>& InBytes, bool bIsPersistent = false);
+
+	using FMemoryArchive::operator<<; // For visibility of the overloads we don't override
+
+	//~ Begin FArchive Interface
+	virtual FArchive& operator<<(FName& Name) override;
+	virtual FArchive& operator<<(FSoftObjectPath& AssetPtr) override;
+	virtual int64 TotalSize() override final;
+	virtual void Serialize(void* Data, int64 Num) override final;
+	virtual FString GetArchiveName() const override;
+	//~ End FArchive Interface
+
+private:
+	const FConcertLocalIdentifierTable* LocalIdentifierTable = nullptr;
+	FConcertLocalIdentifierTable* RewriteIdentifierTable = nullptr;
+	TArray<uint8>& Bytes;
+	TArray<uint8> ScratchBytes;
+};
diff --git a/Engine/Plugins/Developer/Concert/ConcertSync/ConcertSyncCore/Source/ConcertSyncCore/Private/ConcertSyncArchives.cpp b/Engine/Plugins/Developer/Concert/ConcertSync/ConcertSyncCore/Source/ConcertSyncCore/Private/ConcertSyncArchives.cpp
--- a/Engine/Plugins/Developer/Concert/ConcertSync/ConcertSyncCore/Source/ConcertSyncCore/Private/ConcertSyncArchives.cpp
+++ b/Engine/Plugins/Developer/Concert/ConcertSync/ConcertSyncCore/Source/ConcertSyncCore/Private/ConcertSyncArchives.cpp
@@ -311,6 +311,7 @@ FArchive& FConcertSyncObjectWriter::operator<<(FSoftObjectPtr& AssetPtr)
 
 FArchive& FConcertSyncObjectWriter::operator<<(FSoftObjectPath& AssetPtr)
 {
+	// TODO: Serialize via FSoftObjectPath::SerializePath to avoid converting all object paths to FName
 	FName ObjPath;
 	if (bSkipAssets)
 	{
@@ -344,36 +345,47 @@ bool FConcertSyncObjectWriter::ShouldSkipProperty(const FProperty* InProperty) c
 		|| (!ConcertSyncUtil::CanExportProperty(InProperty, !IsFilterEditorOnly()));
 }
 
-FConcertSyncObjectReader::FConcertSyncObjectReader(const FConcertLocalIdentifierTable* InLocalIdentifierTable, FConcertSyncWorldRemapper InWorldRemapper, const FConcertSessionVersionInfo* InVersionInfo, UObject* InObj, const TArray<uint8>& InBytes, const FConcertSyncEncounteredMissingObject& InEncounteredMissingObjectDelegate)
-	: FConcertIdentifierReader(InLocalIdentifierTable, InBytes, /*bIsPersistent*/false)
-	, WorldRemapper(MoveTemp(InWorldRemapper))
-	, EncounteredMissingObjectDelegate(InEncounteredMissingObjectDelegate)
+
+namespace UE::Concert::Private::ConcertSyncArchiveUtil
 {
-	ArIgnoreClassRef = false;
-	ArIgnoreArchetypeRef = false;
-	ArNoDelta = true;
-	//SetWantBinaryPropertySerialization(true);
+
+void InitReaderArchive(FArchive& Ar, const FConcertSessionVersionInfo* InVersionInfo)
+{
+	Ar.ArIgnoreClassRef = false;
+	Ar.ArIgnoreArchetypeRef = false;
+	Ar.ArNoDelta = true;
+	//Ar.SetWantBinaryPropertySerialization(true);
 
 	if (InVersionInfo)
 	{
 		FPackageFileVersion UEVersion(InVersionInfo->FileVersion.FileVersion, (EUnrealEngineObjectUE5Version)InVersionInfo->FileVersion.FileVersionUE5);
 
-		SetUEVer(UEVersion);
-		SetLicenseeUEVer(InVersionInfo->FileVersion.FileVersionLicensee);
-		SetEngineVer(FEngineVersionBase(InVersionInfo->EngineVersion.Major, InVersionInfo->EngineVersion.Minor, InVersionInfo->EngineVersion.Patch, InVersionInfo->EngineVersion.Changelist));
+		Ar.SetUEVer(UEVersion);
+		Ar.SetLicenseeUEVer(InVersionInfo->FileVersion.FileVersionLicensee);
+		Ar.SetEngineVer(FEngineVersionBase(InVersionInfo->EngineVersion.Major, InVersionInfo->EngineVersion.Minor, InVersionInfo->EngineVersion.Patch, InVersionInfo->EngineVersion.Changelist));
 
 		FCustomVersionContainer EngineCustomVersions;
 		for (const FConcertCustomVersionInfo& CustomVersion : InVersionInfo->CustomVersions)
 		{
 			EngineCustomVersions.SetVersion(CustomVersion.Key, CustomVersion.Version, CustomVersion.FriendlyName);
 		}
-		SetCustomVersions(EngineCustomVersions);
+		Ar.SetCustomVersions(EngineCustomVersions);
 	}
 
 	// This is conditional on WITH_EDITORONLY_DATA because the transaction flag is ignored in builds without it and IsTransacting always returns false. So
-	// the writer should only be sending non-editor only properties if the reader does not use editor oly data.
-	SetIsTransacting(WITH_EDITORONLY_DATA);
-	SetFilterEditorOnly(!WITH_EDITORONLY_DATA);
+	// the writer should only be sending non-editor only properties if the reader does not use editor only data.
+	Ar.SetIsTransacting(WITH_EDITORONLY_DATA);
+	Ar.SetFilterEditorOnly(!WITH_EDITORONLY_DATA);
+}
+
+} // namespace UE::Concert::Private::ConcertSyncArchiveUtil
+
+FConcertSyncObjectReader::FConcertSyncObjectReader(const FConcertLocalIdentifierTable* InLocalIdentifierTable, FConcertSyncWorldRemapper InWorldRemapper, const FConcertSessionVersionInfo* InVersionInfo, UObject* InObj, const TArray<uint8>& InBytes, const FConcertSyncEncounteredMissingObject& InEncounteredMissingObjectDelegate)
+	: FConcertIdentifierReader(InLocalIdentifierTable, InBytes, /*bIsPersistent*/false)
+	, WorldRemapper(MoveTemp(InWorldRemapper))
+	, EncounteredMissingObjectDelegate(InEncounteredMissingObjectDelegate)
+{
+	UE::Concert::Private::ConcertSyncArchiveUtil::InitReaderArchive(*this, InVersionInfo);
 
 #if USE_STABLE_LOCALIZATION_KEYS
 	if (GIsEditor && !(ArPortFlags & PPF_DuplicateForPIE))
@@ -505,3 +517,95 @@ FString FConcertSyncObjectReader::GetArchiveName() const
 {
 	return TEXT("FConcertSyncObjectReader");
 }
+
+
+FConcertSyncObjectRewriter::FConcertSyncObjectRewriter(const FConcertLocalIdentifierTable* InLocalIdentifierTable, FConcertLocalIdentifierTable* InRewriteIdentifierTable, const FConcertSessionVersionInfo* InVersionInfo, TArray<uint8>& InBytes)
+	: FConcertIdentifierRewriter(InLocalIdentifierTable, InRewriteIdentifierTable, InBytes, /*bIsPersistent*/false)
+{
+	UE::Concert::Private::ConcertSyncArchiveUtil::InitReaderArchive(*this, InVersionInfo);
+}
+
+void FConcertSyncObjectRewriter::RewriteObject(const UClass* InClass)
+{
+	// Allocate a temporary object to read into while rewriting the data
+	// TODO: This should probably use a pool of some kind
+	UObject* TmpObject = NewObject<UObject>(GetTransientPackage(), InClass);
+	TmpObject->Serialize(*this);
+}
+
+void FConcertSyncObjectRewriter::RewriteProperty(const FProperty* InProp)
+{
+	// Allocate a temporary property to read into while rewriting the data
+	void* TmpPropData = FMemory_Alloca_Aligned(InProp->GetSize(), InProp->GetMinAlignment());
+	InProp->InitializeValue(TmpPropData);
+
+	for (int32 Idx = 0; Idx < InProp->ArrayDim; ++Idx)
+	{
+		InProp->SerializeItem(FStructuredArchiveFromArchive(*this).GetSlot(), (uint8*)TmpPropData + (InProp->ElementSize * Idx));
+	}
+
+	InProp->DestroyValue(TmpPropData);
+}
+
+FArchive& FConcertSyncObjectRewriter::operator<<(UObject*& Obj)
+{
+	FName ObjPath;
+	*this << ObjPath;
+
+	if (ObjPath.IsNone())
+	{
+		Obj = nullptr;
+	}
+	else if (ObjPath != SkipAssetsMarker)
+	{
+		Obj = StaticFindObject(UObject::StaticClass(), nullptr, *ObjPath.ToString());
+	}
+
+	return *this;
+}
+
+FArchive& FConcertSyncObjectRewriter::operator<<(FLazyObjectPtr& LazyObjectPtr)
+{
+	UObject* Obj = nullptr;
+	FUniqueObjectGuid SavedObjectGuid;
+	*this << Obj;
+	*this << SavedObjectGuid;
+
+	// if the resolved object already has an associated Guid, use that instead of the saved one
+	// otherwise used the saved guid since it should refer to the Obj path once its state get applied.
+	FUniqueObjectGuid ObjectGuid = Obj != nullptr ? FUniqueObjectGuid(Obj) : FUniqueObjectGuid();
+	LazyObjectPtr = ObjectGuid.IsValid() ? ObjectGuid : SavedObjectGuid;
+	// technically the saved object guid should be the same as the resolved object guid if any.
+	ensure(!ObjectGuid.IsValid() || ObjectGuid == SavedObjectGuid);
+
+	return *this;
+}
+
+FArchive& FConcertSyncObjectRewriter::operator<<(FObjectPtr& Obj)
+{
+	UObject* RawObj = Obj.Get();
+	*this << RawObj;
+	Obj = RawObj;
+	return *this;
+}
+
+FArchive& FConcertSyncObjectRewriter::operator<<(FSoftObjectPtr& AssetPtr)
+{
+	FSoftObjectPath Obj;
+	*this << Obj;
+	AssetPtr = Obj;
+	return *this;
+}
+
+FArchive& FConcertSyncObjectRewriter::operator<<(FWeakObjectPtr& Value)
+{
+	UObject* Obj = nullptr;
+	*this << Obj;
+	Value = Obj;
+	return *this;
+}
+
+FString FConcertSyncObjectRewriter::GetArchiveName() const
+{
+	return TEXT("FConcertSyncObjectRewriter");
+}
diff --git a/Engine/Plugins/Developer/Concert/ConcertSync/ConcertSyncCore/Source/ConcertSyncCore/Public/ConcertSyncArchives.h b/Engine/Plugins/Developer/Concert/ConcertSync/ConcertSyncCore/Source/ConcertSyncCore/Public/ConcertSyncArchives.h
--- a/Engine/Plugins/Developer/Concert/ConcertSync/ConcertSyncCore/Source/ConcertSyncCore/Public/ConcertSyncArchives.h
+++ b/Engine/Plugins/Developer/Concert/ConcertSync/ConcertSyncCore/Source/ConcertSyncCore/Public/ConcertSyncArchives.h
@@ -79,7 +79,7 @@ private:
 	FConcertSyncRemapObjectPath RemapObjectPathDelegate;
 };
 
-/** Archive for reading objects that can been received from another instance via Concert */
+/** Archive for reading objects that have been received from another instance via Concert */
 class CONCERTSYNCCORE_API FConcertSyncObjectReader : public FConcertIdentifierReader
 {
 public:
@@ -106,6 +106,27 @@ private:
 	FConcertSyncEncounteredMissingObject EncounteredMissingObjectDelegate;
 };
 
+/** Archive for rewriting identifiers (currently names) so that they belong to a different identifier table */
+class CONCERTSYNCCORE_API FConcertSyncObjectRewriter : public FConcertIdentifierRewriter
+{
+public:
+	FConcertSyncObjectRewriter(const FConcertLocalIdentifierTable* InLocalIdentifierTable, FConcertLocalIdentifierTable* InRewriteIdentifierTable, const FConcertSessionVersionInfo* InVersionInfo, TArray<uint8>& InBytes);
+
+	void RewriteObject(const UClass* InClass);
+	void RewriteProperty(const FProperty* InProp);
+
+	using FConcertIdentifierRewriter::operator<<; // For visibility of the overloads we don't override
+
+	//~ Begin FArchive Interface
+	virtual FArchive& operator<<(UObject*& Obj) override;
+	virtual FArchive& operator<<(FLazyObjectPtr& LazyObjectPtr) override;
+	virtual FArchive& operator<<(FObjectPtr& Obj) override;
+	virtual FArchive& operator<<(FSoftObjectPtr& AssetPtr) override;
+	virtual FArchive& operator<<(FWeakObjectPtr& Value) override;
+	virtual FString GetArchiveName() const override;
+	//~ End FArchive Interface
+};
+
 #if UE_ENABLE_INCLUDE_ORDER_DEPRECATED_IN_5_2
 #include "CoreMinimal.h"
 #endif
