Fix client and server module definition typo.

diff --git a/Engine/Plugins/Developer/Concert/ConcertMain/Source/ConcertClient/Private/ConcertClientModule.cpp b/Engine/Plugins/Developer/Concert/ConcertMain/Source/ConcertClient/Private/ConcertClientModule.cpp
--- a/Engine/Plugins/Developer/Concert/ConcertMain/Source/ConcertClient/Private/ConcertClientModule.cpp
+++ b/Engine/Plugins/Developer/Concert/ConcertMain/Source/ConcertClient/Private/ConcertClientModule.cpp
@@ -34,4 +34,4 @@ private:
 	TSharedPtr<IConcertEndpointProvider> EndpointProvider;
 };
 
-IMPLEMENT_MODULE(FConcertClientModule, ConcertClientModule)
+IMPLEMENT_MODULE(FConcertClientModule, ConcertClient)
diff --git a/Engine/Plugins/Developer/Concert/ConcertMain/Source/ConcertServer/Private/ConcertServerModule.cpp b/Engine/Plugins/Developer/Concert/ConcertMain/Source/ConcertServer/Private/ConcertServerModule.cpp
--- a/Engine/Plugins/Developer/Concert/ConcertMain/Source/ConcertServer/Private/ConcertServerModule.cpp
+++ b/Engine/Plugins/Developer/Concert/ConcertMain/Source/ConcertServer/Private/ConcertServerModule.cpp
@@ -31,4 +31,4 @@ private:
 	TSharedPtr<IConcertEndpointProvider> EndpointProvider;
 };
 
-IMPLEMENT_MODULE(FConcertServerModule, ConcertServerModule)
+IMPLEMENT_MODULE(FConcertServerModule, ConcertServer)
