Updated FConcertSyncObjectRewriter to expose support for inline rewriting object paths

Note: The base archive doesn't actually rewrite any paths itself, but a derived archive can by overriding OnObjectSerialized

#rnx

diff --git a/Engine/Plugins/Developer/Concert/ConcertMain/Source/ConcertTransport/Private/IdentifierTable/ConcertTransportArchives.cpp b/Engine/Plugins/Developer/Concert/ConcertMain/Source/ConcertTransport/Private/IdentifierTable/ConcertTransportArchives.cpp
--- a/Engine/Plugins/Developer/Concert/ConcertMain/Source/ConcertTransport/Private/IdentifierTable/ConcertTransportArchives.cpp
+++ b/Engine/Plugins/Developer/Concert/ConcertMain/Source/ConcertTransport/Private/IdentifierTable/ConcertTransportArchives.cpp
@@ -205,33 +205,18 @@ FArchive& FConcertIdentifierRewriter::operator<<(FName& Name)
 		return *this;
 	}
 
-	const int32 OffsetBeforeNameRead = Offset;
+	const int64 OffsetBeforeNameRead = Offset;
 	UE::Concert::Private::ConcertIdentifierArchiveUtil::ReadName(*this, Name, LocalIdentifierTable);
-	const int32 OffsetAfterNameRead = Offset;
+	const int64 OffsetAfterNameRead = Offset;
 
-	// Write the name to the scratch buffer using the rewrite identifier table
+	if (!LocalIdentifierTable->HasMappings())
 	{
-		ScratchBytes.Reset();
-		FMemoryWriter ScratchWriter(ScratchBytes, IsPersistent());
-		UE::Concert::Private::ConcertIdentifierArchiveUtil::WriteName(ScratchWriter, Name, RewriteIdentifierTable);
-	}
-
-	// Patch the scratch buffer into the real buffer
-	// The data must be the same size to avoid serialization issues with tagged data
-	{
-		const int64 OldSerializedSize = OffsetAfterNameRead - OffsetBeforeNameRead;
-		const int64 NewSerializedSize = ScratchBytes.Num();
-		if (ensureAlwaysMsgf(OldSerializedSize == NewSerializedSize, TEXT("FConcertIdentifierRewriter: Rewritten data must be the same size to avoid issues with tagged property serialization! (old: %d, new: %d)"), OldSerializedSize, NewSerializedSize))
-		{
-			FMemory::Memcpy(Bytes.GetData() + OffsetBeforeNameRead, ScratchBytes.GetData(), NewSerializedSize);
-		}
-		else
-		{
-			SetError();
-			return *this;
-		}
+		// No mappings means nothing to rewrite
+		return *this;
 	}
 
+	// Rewrite the data via the scratch buffer
+	RewriteData(OffsetBeforeNameRead, OffsetAfterNameRead - OffsetBeforeNameRead, Name);
 	return *this;
 }
 
diff --git a/Engine/Plugins/Developer/Concert/ConcertMain/Source/ConcertTransport/Public/IdentifierTable/ConcertTransportArchives.h b/Engine/Plugins/Developer/Concert/ConcertMain/Source/ConcertTransport/Public/IdentifierTable/ConcertTransportArchives.h
--- a/Engine/Plugins/Developer/Concert/ConcertMain/Source/ConcertTransport/Public/IdentifierTable/ConcertTransportArchives.h
+++ b/Engine/Plugins/Developer/Concert/ConcertMain/Source/ConcertTransport/Public/IdentifierTable/ConcertTransportArchives.h
@@ -22,6 +22,9 @@ public:
 	virtual FString GetArchiveName() const override;
 	//~ End FArchive Interface
 
+protected:
+	FConcertLocalIdentifierTable* GetLocalIdentifierTable() const { return LocalIdentifierTable; }
+
 private:
 	FConcertLocalIdentifierTable* LocalIdentifierTable;
 };
@@ -40,6 +43,9 @@ public:
 	virtual FString GetArchiveName() const override;
 	//~ End FArchive Interface
 
+protected:
+	const FConcertLocalIdentifierTable* GetLocalIdentifierTable() const { return LocalIdentifierTable; }
+
 private:
 	const FConcertLocalIdentifierTable* LocalIdentifierTable;
 };
@@ -60,6 +66,36 @@ public:
 	virtual FString GetArchiveName() const override;
 	//~ End FArchive Interface
 
+protected:
+	const FConcertLocalIdentifierTable* GetLocalIdentifierTable() const { return LocalIdentifierTable; }
+	FConcertLocalIdentifierTable* GetRewriteIdentifierTable() const { return RewriteIdentifierTable; }
+
+	template <typename T>
+	void RewriteData(const int64 OldDataOffset, const int64 OldDataSize, T& NewData)
+	{
+		// Write the name to the scratch buffer using the rewrite identifier table
+		{
+			ScratchBytes.Reset();
+			FConcertIdentifierWriter ScratchWriter(RewriteIdentifierTable, ScratchBytes, IsPersistent());
+			checkf(ScratchBytes.Num() == 0, TEXT("FConcertIdentifierWriter wrote extra data!"));
+			ScratchWriter << NewData;
+		}
+
+		// Patch the scratch buffer into the real buffer
+		// The data must be the same size to avoid serialization issues with tagged data
+		{
+			const int64 NewDataSize = ScratchBytes.Num();
+			if (ensureAlwaysMsgf(OldDataSize == NewDataSize, TEXT("FConcertIdentifierRewriter: Rewritten data must be the same size to avoid issues with tagged property serialization! (old: %d, new: %d)"), OldDataSize, NewDataSize))
+			{
+				FMemory::Memcpy(Bytes.GetData() + OldDataOffset, ScratchBytes.GetData(), NewDataSize);
+			}
+			else
+			{
+				SetError();
+			}
+		}
+	}
+
 private:
 	const FConcertLocalIdentifierTable* LocalIdentifierTable = nullptr;
 	FConcertLocalIdentifierTable* RewriteIdentifierTable = nullptr;
diff --git a/Engine/Plugins/Developer/Concert/ConcertSync/ConcertSyncCore/Source/ConcertSyncCore/Private/ConcertSyncArchives.cpp b/Engine/Plugins/Developer/Concert/ConcertSync/ConcertSyncCore/Source/ConcertSyncCore/Private/ConcertSyncArchives.cpp
--- a/Engine/Plugins/Developer/Concert/ConcertSync/ConcertSyncCore/Source/ConcertSyncCore/Private/ConcertSyncArchives.cpp
+++ b/Engine/Plugins/Developer/Concert/ConcertSync/ConcertSyncCore/Source/ConcertSyncCore/Private/ConcertSyncArchives.cpp
@@ -551,7 +551,9 @@ void FConcertSyncObjectRewriter::RewriteProperty(const FProperty* InProp)
 FArchive& FConcertSyncObjectRewriter::operator<<(UObject*& Obj)
 {
 	FSoftObjectPath ObjPath;
+	const int64 OffsetBeforeObjectRead = Offset;
 	ObjPath.SerializePath(*this);
+	const int64 OffsetAfterObjectRead = Offset;
 
 	if (ObjPath.IsNull())
 	{
@@ -559,6 +561,8 @@ FArchive& FConcertSyncObjectRewriter::operator<<(UObject*& Obj)
 	}
 	else if (ObjPath != SkipAssetsMarker)
 	{
+		OnObjectSerialized(ObjPath);
+		RewriteData(OffsetBeforeObjectRead, OffsetAfterObjectRead - OffsetBeforeObjectRead, ObjPath);
 		Obj = StaticFindObject(UObject::StaticClass(), nullptr, *ObjPath.ToString());
 	}
 
@@ -598,6 +602,23 @@ FArchive& FConcertSyncObjectRewriter::operator<<(FSoftObjectPtr& AssetPtr)
 	return *this;
 }
 
+FArchive& FConcertSyncObjectRewriter::operator<<(FSoftObjectPath& AssetPtr)
+{
+	FSoftObjectPath ObjPath;
+	const int64 OffsetBeforeObjectRead = Offset;
+	ObjPath.SerializePath(*this);
+	const int64 OffsetAfterObjectRead = Offset;
+
+	if (ObjPath != SkipAssetsMarker)
+	{
+		OnObjectSerialized(ObjPath);
+		RewriteData(OffsetBeforeObjectRead, OffsetAfterObjectRead - OffsetBeforeObjectRead, ObjPath);
+		AssetPtr = ObjPath;
+	}
+
+	return *this;
+}
+
 FArchive& FConcertSyncObjectRewriter::operator<<(FWeakObjectPtr& Value)
 {
 	UObject* Obj = nullptr;
diff --git a/Engine/Plugins/Developer/Concert/ConcertSync/ConcertSyncCore/Source/ConcertSyncCore/Public/ConcertSyncArchives.h b/Engine/Plugins/Developer/Concert/ConcertSync/ConcertSyncCore/Source/ConcertSyncCore/Public/ConcertSyncArchives.h
--- a/Engine/Plugins/Developer/Concert/ConcertSync/ConcertSyncCore/Source/ConcertSyncCore/Public/ConcertSyncArchives.h
+++ b/Engine/Plugins/Developer/Concert/ConcertSync/ConcertSyncCore/Source/ConcertSyncCore/Public/ConcertSyncArchives.h
@@ -125,9 +125,13 @@ public:
 	virtual FArchive& operator<<(FLazyObjectPtr& LazyObjectPtr) override;
 	virtual FArchive& operator<<(FObjectPtr& Obj) override;
 	virtual FArchive& operator<<(FSoftObjectPtr& AssetPtr) override;
+	virtual FArchive& operator<<(FSoftObjectPath& AssetPtr) override;
 	virtual FArchive& operator<<(FWeakObjectPtr& Value) override;
 	virtual FString GetArchiveName() const override;
 	//~ End FArchive Interface
+
+protected:
+	virtual void OnObjectSerialized(FSoftObjectPath& Obj) {}
 };
 
 #if UE_ENABLE_INCLUDE_ORDER_DEPRECATED_IN_5_2
