Run FixIncludePaths on Concert

#rnx

diff --git a/Engine/Plugins/Developer/Concert/ConcertMain/Source/ConcertClient/ConcertClient.Build.cs b/Engine/Plugins/Developer/Concert/ConcertMain/Source/ConcertClient/ConcertClient.Build.cs
--- a/Engine/Plugins/Developer/Concert/ConcertMain/Source/ConcertClient/ConcertClient.Build.cs
+++ b/Engine/Plugins/Developer/Concert/ConcertMain/Source/ConcertClient/ConcertClient.Build.cs
@@ -20,7 +20,9 @@ namespace UnrealBuildTool.Rules
 			PublicIncludePathModuleNames.AddRange(
 				new string[]
 				{
-					"ConcertTransport"
+					"ConcertTransport",
+					"Concert",
+					"GameplayTags",
 				}
 			);
 			PrivateDependencyModuleNames.AddRange(
diff --git a/Engine/Plugins/Developer/Concert/ConcertMain/Source/ConcertClient/Private/ConcertClientSession.h b/Engine/Plugins/Developer/Concert/ConcertMain/Source/ConcertClient/Private/ConcertClientSession.h
--- a/Engine/Plugins/Developer/Concert/ConcertMain/Source/ConcertClient/Private/ConcertClientSession.h
+++ b/Engine/Plugins/Developer/Concert/ConcertMain/Source/ConcertClient/Private/ConcertClientSession.h
@@ -6,7 +6,7 @@
 #include "ConcertSession.h"
 #include "ConcertMessages.h"
 #include "Containers/Ticker.h"
-#include "GameplayTags/Classes/GameplayTagContainer.h"
+#include "GameplayTagContainer.h"
 
 class IConcertLocalEndpoint;
 struct FConcertClientSettings;
diff --git a/Engine/Plugins/Developer/Concert/ConcertMain/Source/ConcertClient/Public/ConcertClientSettings.h b/Engine/Plugins/Developer/Concert/ConcertMain/Source/ConcertClient/Public/ConcertClientSettings.h
--- a/Engine/Plugins/Developer/Concert/ConcertMain/Source/ConcertClient/Public/ConcertClientSettings.h
+++ b/Engine/Plugins/Developer/Concert/ConcertMain/Source/ConcertClient/Public/ConcertClientSettings.h
@@ -9,7 +9,7 @@
 #include "Misc/EngineVersion.h"
 #include "ConcertTransportSettings.h"
 #include "ConcertVersion.h"
-#include "GameplayTags/Classes/GameplayTagContainer.h"
+#include "GameplayTagContainer.h"
 
 #include "ConcertClientSettings.generated.h"
 
diff --git a/Engine/Plugins/Developer/Concert/ConcertMain/Source/ConcertClient/Public/IConcertClient.h b/Engine/Plugins/Developer/Concert/ConcertMain/Source/ConcertClient/Public/IConcertClient.h
--- a/Engine/Plugins/Developer/Concert/ConcertMain/Source/ConcertClient/Public/IConcertClient.h
+++ b/Engine/Plugins/Developer/Concert/ConcertMain/Source/ConcertClient/Public/IConcertClient.h
@@ -2,7 +2,7 @@
 
 #pragma once
 
-#include "Concert/Public/IConcertSession.h"
+#include "IConcertSession.h"
 #include "CoreMinimal.h"
 #include "Async/Future.h"
 #include "IConcertSession.h"
diff --git a/Engine/Plugins/Developer/Concert/ConcertMain/Source/ConcertServer/Private/ConcertServer.cpp b/Engine/Plugins/Developer/Concert/ConcertMain/Source/ConcertServer/Private/ConcertServer.cpp
--- a/Engine/Plugins/Developer/Concert/ConcertMain/Source/ConcertServer/Private/ConcertServer.cpp
+++ b/Engine/Plugins/Developer/Concert/ConcertMain/Source/ConcertServer/Private/ConcertServer.cpp
@@ -14,16 +14,15 @@
 #include "Algo/AnyOf.h"
 #include "Misc/App.h"
 #include "Misc/Paths.h"
-
-#include "Runtime/Launch/Resources/Version.h"
-#include "StructSerializer.h"
-#include "StructDeserializer.h"
-#include "Backends/JsonStructSerializerBackend.h"
 #include "Backends/JsonStructDeserializerBackend.h"
+#include "Backends/JsonStructSerializerBackend.h"
 #include "HAL/FileManager.h"
-
+#include "StructDeserializer.h"
+#include "StructSerializer.h"
 #include "Templates/NonNullPointer.h"
 
+#include "Runtime/Launch/Resources/Version.h"
+
 #define LOCTEXT_NAMESPACE "ConcertServer"
 
 namespace ConcertServerUtil
