Ensure that hardcoded names and numeric suffixes can still be rewritten

This can be needed for object path remapping, where the name may change from hardcoded to not, or the numeric suffix of the name may change

#rnx

diff --git a/Engine/Plugins/Developer/Concert/ConcertMain/Source/ConcertTransport/Private/IdentifierTable/ConcertTransportArchives.cpp b/Engine/Plugins/Developer/Concert/ConcertMain/Source/ConcertTransport/Private/IdentifierTable/ConcertTransportArchives.cpp
--- a/Engine/Plugins/Developer/Concert/ConcertMain/Source/ConcertTransport/Private/IdentifierTable/ConcertTransportArchives.cpp
+++ b/Engine/Plugins/Developer/Concert/ConcertMain/Source/ConcertTransport/Private/IdentifierTable/ConcertTransportArchives.cpp
@@ -11,13 +11,15 @@ enum class EConcertIdentifierSource : uint8
 	/** Plain string value (no suffix) */
 	PlainString,
 	/** Hardcoded FName index value (see MAX_NETWORKED_HARDCODED_NAME) */
-	HardcodedIndex,
+	HardcodedIndex_PackedInt,
 	/** Local identifier table index value (see FConcertLocalIdentifierTable) */
 	LocalIdentifierTableIndex_PackedInt,
 	/** Local identifier table index value (see FConcertLocalIdentifierTable) */
 	LocalIdentifierTableIndex_FixedSizeInt32,
 	/** FName table index value (non-portable!) */
 	FNameTableIndex_FixedSizeUInt32,
+	/** Hardcoded FName index value (see MAX_NETWORKED_HARDCODED_NAME) */
+	HardcodedIndex_FixedSizeInt32,
 };
 
 void WriteName(FArchive& Ar, const FName& Name, FConcertLocalIdentifierTable* LocalIdentifierTable)
@@ -34,22 +36,35 @@ void WriteName(FArchive& Ar, const FName& Name, FConcertLocalIdentifierTable* Lo
 		Ar.SerializeIntPacked(UnsignedIndex);
 	};
 
+	bool bSerializeNumberAsPacked = true;
 	if (const EName* Ename = Name.ToEName(); Ename && ShouldReplicateAsInteger(*Ename, Name))
 	{
-		SerializeConcertIdentifierSource(EConcertIdentifierSource::HardcodedIndex);
-		SerializeIndexValue((int32)*Ename);
+		if (LocalIdentifierTable)
+		{
+			SerializeConcertIdentifierSource(EConcertIdentifierSource::HardcodedIndex_FixedSizeInt32);
+			int32 ENameInt = (int32)*Ename;
+			Ar << ENameInt; // Note: Don't serialize as a packed int so that the data size remains consistent in case this data gets rewritten (see FConcertIdentifierRewriter)
+			bSerializeNumberAsPacked = false;
+		}
+		else
+		{
+			SerializeConcertIdentifierSource(EConcertIdentifierSource::HardcodedIndex_PackedInt);
+			SerializeIndexValue((int32)*Ename);
+		}
 	}
 	else if (LocalIdentifierTable == FConcertLocalIdentifierTable::ForceFNameTableIndex)
 	{
 		SerializeConcertIdentifierSource(EConcertIdentifierSource::FNameTableIndex_FixedSizeUInt32);
 		uint32 FNameTableIndex = Name.GetDisplayIndex().ToUnstableInt();
 		Ar << FNameTableIndex; // Note: Don't serialize as a packed int so that the data size remains consistent in case this data gets rewritten (see FConcertIdentifierRewriter)
+		bSerializeNumberAsPacked = false;
 	}
 	else if (LocalIdentifierTable)
 	{
 		SerializeConcertIdentifierSource(EConcertIdentifierSource::LocalIdentifierTableIndex_FixedSizeInt32);
 		int32 IdentifierTableIndex = LocalIdentifierTable->MapName(Name);
 		Ar << IdentifierTableIndex; // Note: Don't serialize as a packed int so that the data size remains consistent in case this data gets rewritten (see FConcertIdentifierRewriter)
+		bSerializeNumberAsPacked = false;
 	}
 	else
 	{
@@ -59,7 +74,14 @@ void WriteName(FArchive& Ar, const FName& Name, FConcertLocalIdentifierTable* Lo
 	}
 	
 	int32 NameNumber = Name.GetNumber();
-	SerializeIndexValue(NameNumber);
+	if (bSerializeNumberAsPacked)
+	{
+		SerializeIndexValue(NameNumber);
+	}
+	else
+	{
+		Ar << NameNumber;
+	}
 }
 
 void ReadName(FArchive& Ar, FName& Name, const FConcertLocalIdentifierTable* LocalIdentifierTable)
@@ -76,6 +98,7 @@ void ReadName(FArchive& Ar, FName& Name, const FConcertLocalIdentifierTable* Loc
 	EConcertIdentifierSource Source;
 	Ar.Serialize(&Source, sizeof(EConcertIdentifierSource));
 
+	bool bSerializeNumberAsPacked = true;
 	switch (Source)
 	{
 	case EConcertIdentifierSource::PlainString:
@@ -86,13 +109,22 @@ void ReadName(FArchive& Ar, FName& Name, const FConcertLocalIdentifierTable* Loc
 		}
 		break;
 
-	case EConcertIdentifierSource::HardcodedIndex:
+	case EConcertIdentifierSource::HardcodedIndex_PackedInt:
 		{
 			const int32 HardcodedIndex = SerializeIndexValue();
 			Name = EName(HardcodedIndex);
 		}
 		break;
 
+	case EConcertIdentifierSource::HardcodedIndex_FixedSizeInt32:
+		{
+			int32 HardcodedIndex = INDEX_NONE;
+			Ar << HardcodedIndex;
+			Name = EName(HardcodedIndex);
+			bSerializeNumberAsPacked = false;
+		}
+		break;
+
 	case EConcertIdentifierSource::LocalIdentifierTableIndex_PackedInt:
 		{
 			const int32 IdentifierTableIndex = SerializeIndexValue();
@@ -113,6 +145,7 @@ void ReadName(FArchive& Ar, FName& Name, const FConcertLocalIdentifierTable* Loc
 				Ar.SetError();
 				return;
 			}
+			bSerializeNumberAsPacked = false;
 		}
 		break;
 
@@ -121,6 +154,7 @@ void ReadName(FArchive& Ar, FName& Name, const FConcertLocalIdentifierTable* Loc
 			uint32 FNameTableIndex = 0;
 			Ar << FNameTableIndex;
 			Name = FName::CreateFromDisplayId(FNameEntryId::FromUnstableInt(FNameTableIndex), 0);
+			bSerializeNumberAsPacked = false;
 		}
 		break;
 
@@ -129,7 +163,15 @@ void ReadName(FArchive& Ar, FName& Name, const FConcertLocalIdentifierTable* Loc
 		break;
 	}
 
-	const int32 NameNumber = SerializeIndexValue();
+	int32 NameNumber = 0;
+	if (bSerializeNumberAsPacked)
+	{
+		NameNumber = SerializeIndexValue();
+	}
+	else
+	{
+		Ar << NameNumber;
+	}
 	Name.SetNumber(NameNumber);
 }
 
