[Concert] Adapt ConcertServer to have a ParticipantCanJoinSession delegate for MultiUser Apps

Adapting the ConcertServer to have a new delegate ParticipantCanJoinSession that will be executed (if bound) when a participant attempts to join a session. The participants ConcertEndpointId, ConcertClientInfo and the SessionId are passed along to a bound delegate function. This change supports custom MultiUser Apps.

diff --git a/Engine/Plugins/Developer/Concert/ConcertMain/Source/ConcertServer/Private/ConcertServer.cpp b/Engine/Plugins/Developer/Concert/ConcertMain/Source/ConcertServer/Private/ConcertServer.cpp
--- a/Engine/Plugins/Developer/Concert/ConcertMain/Source/ConcertServer/Private/ConcertServer.cpp
+++ b/Engine/Plugins/Developer/Concert/ConcertMain/Source/ConcertServer/Private/ConcertServer.cpp
@@ -690,6 +690,11 @@ FOnConcertServerStartup& FConcertServer::OnConcertServerStartup()
 	return OnConcertServerStartupDelegate;
 }
 
+FOnConcertParticipantCanJoinSession& FConcertServer::OnConcertParticipantCanJoinSession()
+{
+	return OnConcertParticipantCanJoinSessionDelegate;
+}
+
 TArray<FConcertSessionInfo> FConcertServer::GetLiveSessionInfos() const
 {
 	TArray<FConcertSessionInfo> SessionsInfo;
@@ -1127,7 +1132,7 @@ TFuture<FConcertAdmin_SessionInfoResponse> FConcertServer::HandleFindSessionRequ
 	// Find the session requested
 	TSharedPtr<IConcertServerSession> ServerSession = GetLiveSession(Message->SessionId);
 	const TCHAR* ServerSessionNamePtr = ServerSession ? *ServerSession->GetName() : TEXT("<unknown>");
-	if (CanJoinSession(ServerSession, Message->SessionSettings, Message->VersionInfo, &ResponseData.Reason))
+	if (CanJoinSession(ServerSession, Message->SessionSettings, Message->VersionInfo, Message->ConcertEndpointId, Message->OwnerClientInfo, &ResponseData.Reason))
 	{
 		ResponseData.ResponseCode = EConcertResponseCode::Success;
 		ResponseData.SessionInfo = ServerSession->GetSessionInfo();
@@ -1496,7 +1501,7 @@ TFuture<FConcertAdmin_GetSessionActivitiesResponse> FConcertServer::HandleGetSes
 	return FConcertAdmin_GetSessionActivitiesResponse::AsFuture(MoveTemp(ResponseData));
 }
 
-bool FConcertServer::CanJoinSession(const TSharedPtr<IConcertServerSession>& ServerSession, const FConcertSessionSettings& SessionSettings, const FConcertSessionVersionInfo& SessionVersionInfo, FText* OutFailureReason)
+bool FConcertServer::CanJoinSession(const TSharedPtr<IConcertServerSession>& ServerSession, const FConcertSessionSettings& SessionSettings, const FConcertSessionVersionInfo& SessionVersionInfo, const FGuid& EndpointId, const FConcertClientInfo& ClientInfo, FText* OutFailureReason)
 {
 	if (!ServerSession)
 	{
@@ -1507,6 +1512,14 @@ bool FConcertServer::CanJoinSession(const TSharedPtr<IConcertServerSession>& Ser
 		return false;
 	}
 
+	if (OnConcertParticipantCanJoinSessionDelegate.IsBound())
+	{
+		if (!OnConcertParticipantCanJoinSessionDelegate.Execute(ServerSession->GetId(), EndpointId, ClientInfo, OutFailureReason))
+		{
+			return false;
+		}
+	}
+
 	if (Settings->ServerSettings.bIgnoreSessionSettingsRestriction)
 	{
 		return true;
diff --git a/Engine/Plugins/Developer/Concert/ConcertMain/Source/ConcertServer/Private/ConcertServer.h b/Engine/Plugins/Developer/Concert/ConcertMain/Source/ConcertServer/Private/ConcertServer.h
--- a/Engine/Plugins/Developer/Concert/ConcertMain/Source/ConcertServer/Private/ConcertServer.h
+++ b/Engine/Plugins/Developer/Concert/ConcertMain/Source/ConcertServer/Private/ConcertServer.h
@@ -53,11 +53,13 @@ public:
 	virtual bool DestroySession(const FGuid& SessionId, FText& OutFailureReason) override;
 	virtual FOnConcertServerSessionStartup& OnConcertServerSessionStartup() override;
 	virtual FOnConcertServerStartup& OnConcertServerStartup() override;
+	virtual FOnConcertParticipantCanJoinSession& OnConcertParticipantCanJoinSession() override;
 
 protected:
 
 	FOnConcertServerSessionStartup OnConcertServerSessionStartupDelegate;
 	FOnConcertServerStartup OnConcertServerStartupDelegate;
+	FOnConcertParticipantCanJoinSession OnConcertParticipantCanJoinSessionDelegate;
 
 private:
 	/** Returns the root dir where the servers keeps the its internally created repositories (When the caller doesn't provide the paths). */
@@ -140,8 +142,8 @@ private:
 	 */
 	void ArchiveOfflineSessions(const FConcertServerSessionRepository& InRepository);
 
-	/** */
-	bool CanJoinSession(const TSharedPtr<IConcertServerSession>& ServerSession, const FConcertSessionSettings& SessionSettings, const FConcertSessionVersionInfo& SessionVersionInfo, FText* OutFailureReason = nullptr);
+	/** Check whether a connecting participant can join the session. */
+	bool CanJoinSession(const TSharedPtr<IConcertServerSession>& ServerSession, const FConcertSessionSettings& SessionSettings, const FConcertSessionVersionInfo& SessionVersionInfo, const FGuid& EndpointId, const FConcertClientInfo& ClientInfo, FText* OutFailureReason = nullptr);
 
 	/**  Validate that the request to delete or modify a session comes from the owner of that session. */
 	bool IsRequestFromSessionOwner(const TSharedPtr<IConcertServerSession>& Session, const FString& FromUserName, const FString& FromDeviceName) const;
diff --git a/Engine/Plugins/Developer/Concert/ConcertMain/Source/ConcertServer/Public/IConcertServer.h b/Engine/Plugins/Developer/Concert/ConcertMain/Source/ConcertServer/Public/IConcertServer.h
--- a/Engine/Plugins/Developer/Concert/ConcertMain/Source/ConcertServer/Public/IConcertServer.h
+++ b/Engine/Plugins/Developer/Concert/ConcertMain/Source/ConcertServer/Public/IConcertServer.h
@@ -7,9 +7,11 @@
 #include "IConcertSession.h"
 
 class UConcertServerConfig;
+struct FConcertClientInfo;
 
 DECLARE_MULTICAST_DELEGATE_OneParam(FOnConcertServerSessionStartup, TWeakPtr<IConcertServerSession>);
 DECLARE_MULTICAST_DELEGATE(FOnConcertServerStartup);
+DECLARE_DELEGATE_RetVal_FourParams(bool, FOnConcertParticipantCanJoinSession, const FGuid& /*SessionId*/, const FGuid& /*EndpointId*/, const FConcertClientInfo& /*ClientInfo*/, FText* /*OutFailureReason*/);
 
 /** Interface for Concert server */
 class IConcertServer
@@ -51,6 +53,9 @@ public:
 	
 	/** Callback when a remote admin endpoint connection changes. */
 	virtual FOnConcertRemoteEndpointConnectionChanged& OnRemoteEndpointConnectionChanged() = 0;
+
+	/** Callback to check whether a participant should be allowed to join a session. */
+	virtual FOnConcertParticipantCanJoinSession& OnConcertParticipantCanJoinSession() = 0;
 	
 	/**
 	 * Gets the address of a remote admin endpoint, i.e. a client that is sending FConcertEndpointDiscoveryEvents. 
