Fix circular include for Horde issue 309764

diff --git a/Engine/Plugins/Developer/Concert/ConcertMain/Source/Concert/Public/IConcertSession.h b/Engine/Plugins/Developer/Concert/ConcertMain/Source/Concert/Public/IConcertSession.h
--- a/Engine/Plugins/Developer/Concert/ConcertMain/Source/Concert/Public/IConcertSession.h
+++ b/Engine/Plugins/Developer/Concert/ConcertMain/Source/Concert/Public/IConcertSession.h
@@ -5,7 +5,6 @@
 #include "ConcertMessageContext.h"
 #include "CoreMinimal.h"
 #include "IConcertEndpoint.h"
-#include "IConcertSession.h"
 #include "IConcertSessionHandler.h"
 #include "ConcertMessages.h"
 #include "Containers/ArrayBuilder.h"
