Fix a number of plugins that are missing some references / supported programs

#rnx

diff --git a/Engine/Plugins/Developer/Concert/ConcertMain/ConcertMain.uplugin b/Engine/Plugins/Developer/Concert/ConcertMain/ConcertMain.uplugin
--- a/Engine/Plugins/Developer/Concert/ConcertMain/ConcertMain.uplugin
+++ b/Engine/Plugins/Developer/Concert/ConcertMain/ConcertMain.uplugin
@@ -15,31 +15,31 @@
 	"Hidden": true,
 	"IsBetaVersion": true,
 	"Installed": false,
-	"SupportedPrograms": [ "UnrealMultiUserServer", "UnrealMultiUserSlateServer", "UnrealRecoverySvc", "CrashReportClientEditor" ],
+	"SupportedPrograms": [ "UnrealMultiUserServer", "CoopMultiUserServer", "UnrealMultiUserSlateServer", "UnrealRecoverySvc", "CrashReportClientEditor", "LiveLinkHub" ],
 	"Modules": [
 		{
 			"Name": "Concert",
 			"Type": "UncookedOnly",
 			"LoadingPhase": "PreDefault",
-			"ProgramAllowList": [ "UnrealMultiUserServer", "UnrealMultiUserSlateServer", "UnrealRecoverySvc", "CrashReportClientEditor" ]
+			"ProgramAllowList": [ "UnrealMultiUserServer", "CoopMultiUserServer", "UnrealMultiUserSlateServer", "UnrealRecoverySvc", "CrashReportClientEditor", "LiveLinkHub" ]
 		},
 		{
 			"Name": "ConcertClient",
 			"Type": "UncookedOnly",
 			"LoadingPhase": "PreDefault",
-			"ProgramDenyList": [ "UnrealMultiUserServer", "UnrealMultiUserSlateServer", "UnrealRecoverySvc", "CrashReportClientEditor" ]
+			"ProgramDenyList": [ "UnrealMultiUserServer", "CoopMultiUserServer", "UnrealMultiUserSlateServer", "UnrealRecoverySvc", "CrashReportClientEditor", "LiveLinkHub" ]
 		},
 		{
 			"Name": "ConcertServer",
 			"Type": "UncookedOnly",
 			"LoadingPhase": "PreDefault",
-			"ProgramAllowList": [ "UnrealMultiUserServer", "UnrealMultiUserSlateServer", "UnrealRecoverySvc", "CrashReportClientEditor" ]
+			"ProgramAllowList": [ "UnrealMultiUserServer", "CoopMultiUserServer", "UnrealMultiUserSlateServer", "UnrealRecoverySvc", "CrashReportClientEditor", "LiveLinkHub" ]
 		},
 		{
 			"Name": "ConcertTransport",
 			"Type": "UncookedOnly",
 			"LoadingPhase": "PreDefault",
-			"ProgramAllowList": [ "UnrealMultiUserServer", "UnrealMultiUserSlateServer", "UnrealRecoverySvc", "CrashReportClientEditor" ]
+			"ProgramAllowList": [ "UnrealMultiUserServer", "CoopMultiUserServer", "UnrealMultiUserSlateServer", "UnrealRecoverySvc", "CrashReportClientEditor", "LiveLinkHub" ]
 		}
 	]
 }
diff --git a/Engine/Plugins/Developer/Concert/ConcertSync/ConcertSyncClient/ConcertSyncClient.uplugin b/Engine/Plugins/Developer/Concert/ConcertSync/ConcertSyncClient/ConcertSyncClient.uplugin
--- a/Engine/Plugins/Developer/Concert/ConcertSync/ConcertSyncClient/ConcertSyncClient.uplugin
+++ b/Engine/Plugins/Developer/Concert/ConcertSync/ConcertSyncClient/ConcertSyncClient.uplugin
@@ -15,6 +15,7 @@
 	"Hidden": true,
 	"IsBetaVersion": true,
 	"Installed": false,
+	"SupportedPrograms": [ "LiveLinkHub" ],
 	"Modules": [
 		{
 			"Name": "ConcertSyncClient",
diff --git a/Engine/Plugins/Developer/Concert/ConcertSync/ConcertSyncCore/ConcertSyncCore.uplugin b/Engine/Plugins/Developer/Concert/ConcertSync/ConcertSyncCore/ConcertSyncCore.uplugin
--- a/Engine/Plugins/Developer/Concert/ConcertSync/ConcertSyncCore/ConcertSyncCore.uplugin
+++ b/Engine/Plugins/Developer/Concert/ConcertSync/ConcertSyncCore/ConcertSyncCore.uplugin
@@ -15,13 +15,13 @@
 	"Hidden": true,
 	"IsBetaVersion": true,
 	"Installed": false,
-	"SupportedPrograms": [ "UnrealMultiUserServer", "UnrealMultiUserSlateServer", "UnrealRecoverySvc", "CrashReportClientEditor" ],
+	"SupportedPrograms": [ "UnrealMultiUserServer", "CoopMultiUserServer", "UnrealMultiUserSlateServer", "UnrealRecoverySvc", "CrashReportClientEditor", "LiveLinkHub" ],
 	"Modules": [
 		{
 			"Name": "ConcertSyncCore",
 			"Type": "UncookedOnly",
 			"LoadingPhase": "PreDefault",
-			"ProgramAllowList": [ "UnrealMultiUserServer", "UnrealMultiUserSlateServer", "UnrealRecoverySvc", "CrashReportClientEditor" ]
+			"ProgramAllowList": [ "UnrealMultiUserServer", "CoopMultiUserServer", "UnrealMultiUserSlateServer", "UnrealRecoverySvc", "CrashReportClientEditor", "LiveLinkHub" ]
 		}
 	],
 	"Plugins": [
diff --git a/Engine/Plugins/Developer/Concert/ConcertSync/ConcertSyncServer/ConcertSyncServer.uplugin b/Engine/Plugins/Developer/Concert/ConcertSync/ConcertSyncServer/ConcertSyncServer.uplugin
--- a/Engine/Plugins/Developer/Concert/ConcertSync/ConcertSyncServer/ConcertSyncServer.uplugin
+++ b/Engine/Plugins/Developer/Concert/ConcertSync/ConcertSyncServer/ConcertSyncServer.uplugin
@@ -15,13 +15,13 @@
 	"Hidden": true,
 	"IsBetaVersion": true,
 	"Installed": false,
-	"SupportedPrograms": [ "UnrealMultiUserServer", "UnrealMultiUserSlateServer", "UnrealRecoverySvc", "CrashReportClientEditor" ],
+	"SupportedPrograms": [ "UnrealMultiUserServer", "CoopMultiUserServer", "UnrealMultiUserSlateServer", "UnrealRecoverySvc", "CrashReportClientEditor" ],
 	"Modules": [
 		{
 			"Name": "ConcertSyncServer",
 			"Type": "UncookedOnly",
 			"LoadingPhase": "Default",
-			"ProgramAllowList": [ "UnrealMultiUserServer", "UnrealMultiUserSlateServer", "UnrealRecoverySvc", "CrashReportClientEditor" ]
+			"ProgramAllowList": [ "UnrealMultiUserServer", "CoopMultiUserServer", "UnrealMultiUserSlateServer", "UnrealRecoverySvc", "CrashReportClientEditor" ]
 		}
 	],
 	"Plugins": [
diff --git a/Engine/Plugins/Editor/ContentBrowser/ContentBrowserFileDataSource/ContentBrowserFileDataSource.uplugin b/Engine/Plugins/Editor/ContentBrowser/ContentBrowserFileDataSource/ContentBrowserFileDataSource.uplugin
--- a/Engine/Plugins/Editor/ContentBrowser/ContentBrowserFileDataSource/ContentBrowserFileDataSource.uplugin
+++ b/Engine/Plugins/Editor/ContentBrowser/ContentBrowserFileDataSource/ContentBrowserFileDataSource.uplugin
@@ -14,11 +14,12 @@
 	"CanContainContent" : false,
 	"IsBetaVersion" : false,
 	"Installed" : false,
+	"SupportedPrograms": ["LiveLinkHub"],
 	"Modules" :
 	[
 		{
 			"Name" : "ContentBrowserFileDataSource",
-			"Type" : "Editor",
+			"Type" : "EditorAndProgram",
 			"LoadingPhase" : "Default"
 		}
 	]
diff --git a/Engine/Plugins/Experimental/ChaosFlesh/ChaosFlesh.uplugin b/Engine/Plugins/Experimental/ChaosFlesh/ChaosFlesh.uplugin
--- a/Engine/Plugins/Experimental/ChaosFlesh/ChaosFlesh.uplugin
+++ b/Engine/Plugins/Experimental/ChaosFlesh/ChaosFlesh.uplugin
@@ -75,6 +75,10 @@
 			"SupportedTargetPlatforms": [
 				"Win64"
 			]
-        }
+        },
+		{
+			"Name": "TetMeshing",
+			"Enabled": true
+		}
 	]
 }
diff --git a/Engine/Plugins/Experimental/PythonScriptPlugin/PythonScriptPlugin.uplugin b/Engine/Plugins/Experimental/PythonScriptPlugin/PythonScriptPlugin.uplugin
--- a/Engine/Plugins/Experimental/PythonScriptPlugin/PythonScriptPlugin.uplugin
+++ b/Engine/Plugins/Experimental/PythonScriptPlugin/PythonScriptPlugin.uplugin
@@ -14,6 +14,7 @@
 	"CanContainContent": true,
 	"IsBetaVersion": true,
 	"Installed": false,
+	"SupportedPrograms": ["LiveLinkHub"],
 	"Modules": [
 		{
 			"Name": "PythonScriptPluginPreload",
@@ -30,8 +31,10 @@
 		{
 			"Name": "ContentBrowserFileDataSource",
 			"Enabled": true,
+			"ProgramAllowList": [ "LiveLinkHub" ],
 			"TargetAllowList": [
-				"Editor"
+				"Editor",
+				"Program"
 			]
 		}
 	]
diff --git a/Engine/Plugins/Runtime/Database/SQLiteCore/SQLiteCore.uplugin b/Engine/Plugins/Runtime/Database/SQLiteCore/SQLiteCore.uplugin
--- a/Engine/Plugins/Runtime/Database/SQLiteCore/SQLiteCore.uplugin
+++ b/Engine/Plugins/Runtime/Database/SQLiteCore/SQLiteCore.uplugin
@@ -14,7 +14,7 @@
 	"CanContainContent": false,
 	"IsBetaVersion": false,
 	"Installed": false,
-	"SupportedPrograms": [ "UnrealMultiUserServer", "UnrealMultiUserSlateServer", "UnrealRecoverySvc" ],
+	"SupportedPrograms": [ "UnrealMultiUserServer", "CoopMultiUserServer", "UnrealMultiUserSlateServer", "UnrealRecoverySvc", "LiveLinkHub" ],
 	"Modules": [
 		{
 			"Name": "SQLiteCore",
diff --git a/Engine/Source/Programs/UnrealBuildTool/System/ModuleDescriptor.cs b/Engine/Source/Programs/UnrealBuildTool/System/ModuleDescriptor.cs
--- a/Engine/Source/Programs/UnrealBuildTool/System/ModuleDescriptor.cs
+++ b/Engine/Source/Programs/UnrealBuildTool/System/ModuleDescriptor.cs
@@ -88,6 +88,11 @@ namespace UnrealBuildTool
 		/// Loaded only by clients and editor (editor can run PIE which is kinda a commandlet)
 		/// </summary>
 		ClientOnlyNoCommandlet,
+
+		/// <summary>
+		/// External module, should never be loaded automatically only referenced
+		/// </summary>
+		External,
 	}
 
 	/// <summary>
