Add analytics to Multi-user session start

diff --git a/Engine/Plugins/Developer/Concert/ConcertMain/Source/ConcertClient/Private/ConcertClientSession.cpp b/Engine/Plugins/Developer/Concert/ConcertMain/Source/ConcertClient/Private/ConcertClientSession.cpp
--- a/Engine/Plugins/Developer/Concert/ConcertMain/Source/ConcertClient/Private/ConcertClientSession.cpp
+++ b/Engine/Plugins/Developer/Concert/ConcertMain/Source/ConcertClient/Private/ConcertClientSession.cpp
@@ -9,6 +9,8 @@
 
 #include "Scratchpad/ConcertScratchpad.h"
 
+#include "EngineAnalytics.h"
+
 #include "Containers/Ticker.h"
 #include "Misc/Paths.h"
 #include "Stats/Stats.h"
@@ -34,6 +36,21 @@ FConcertClientSession::~FConcertClientSession()
 	check(!SessionTick.IsValid());
 }
 
+namespace UE::ConcertClientSession::Private
+{
+	void SendAnalyticsSessionStarted(const FString& SessionIdString)
+	{
+		if (!FEngineAnalytics::IsAvailable())
+		{
+			return;
+		}
+		TArray<FAnalyticsEventAttribute> EventAttributes;
+		EventAttributes.Add(FAnalyticsEventAttribute(TEXT("SessionID"), SessionIdString));
+
+		FEngineAnalytics::GetProvider().RecordEvent(TEXT("Usage.MultiUser.SessionStarted"), EventAttributes);
+	}
+} // namespace UE::ConcertClientSession::Private
+
 void FConcertClientSession::Startup()
 {
 	// if the session tick isn't valid we haven't started
@@ -63,7 +80,9 @@ void FConcertClientSession::Startup()
 			return true;
 		});
 
-		UE_LOG(LogConcert, Display, TEXT("Initialized Concert session '%s' (Id: %s, Owner: %s)."), *SessionInfo.SessionName, *SessionInfo.SessionId.ToString(), *SessionInfo.OwnerUserName);
+		FString SessionIdString = SessionInfo.SessionId.ToString();
+		UE_LOG(LogConcert, Display, TEXT("Initialized Concert session '%s' (Id: %s, Owner: %s)."), *SessionInfo.SessionName, *SessionIdString, *SessionInfo.OwnerUserName);
+		UE::ConcertClientSession::Private::SendAnalyticsSessionStarted(SessionIdString);
 	}
 }
 
