Remove 5.2 deprecated functions from Multi-user interface

diff --git a/Engine/Plugins/Developer/Concert/ConcertMain/Source/Concert/Public/IConcertSession.h b/Engine/Plugins/Developer/Concert/ConcertMain/Source/Concert/Public/IConcertSession.h
--- a/Engine/Plugins/Developer/Concert/ConcertMain/Source/Concert/Public/IConcertSession.h
+++ b/Engine/Plugins/Developer/Concert/ConcertMain/Source/Concert/Public/IConcertSession.h
@@ -391,27 +391,6 @@ public:
 	/** Disconnect  gracefully from the server session */
 	virtual void Disconnect() = 0;
 
-	/** Resume live-updates for this session (must be paired with a call to Suspend) */
-	UE_DEPRECATED(5.2,"Use SetSendReceiveState to change live update status.")
-	virtual void Resume()
-	{
-		SetSendReceiveState(EConcertSendReceiveState::Default);
-	};
-
-	/** Suspend live-updates for this session */
-	UE_DEPRECATED(5.2,"Use SetSendReceiveState to change live update status.")
-	virtual void Suspend()
-	{
-		SetSendReceiveState(EConcertSendReceiveState::SendOnly);
-	}
-
-	/** Does this session currently have live-updates suspended? */
-	UE_DEPRECATED(5.2,"Use GetSendReceiveState to query the live update status.")
-	virtual bool IsSuspended() const
-	{
-		return GetSendReceiveState() == EConcertSendReceiveState::SendOnly;
-	}
-
 	/** Manager for helping sequence custom events on client sessions. */
 	virtual FConcertSequencedCustomEventManager& GetSequencedEventManager() = 0;
 
diff --git a/Engine/Plugins/Developer/Concert/ConcertMain/Source/ConcertClient/Public/IConcertClient.h b/Engine/Plugins/Developer/Concert/ConcertMain/Source/ConcertClient/Public/IConcertClient.h
--- a/Engine/Plugins/Developer/Concert/ConcertMain/Source/ConcertClient/Public/IConcertClient.h
+++ b/Engine/Plugins/Developer/Concert/ConcertMain/Source/ConcertClient/Public/IConcertClient.h
@@ -364,24 +364,6 @@ public:
 	 */
 	virtual void DisconnectSession() = 0;
 
-	UE_DEPRECATED(5.2,"Use SetSendReceiveState to change live update status.")
-	virtual void ResumeSession()
-	{
-		SetSendReceiveState(EConcertSendReceiveState::Default);
-	}
-
-	UE_DEPRECATED(5.2,"Use SetSendReceiveState to change live update status.")
-	virtual void SuspendSession()
-	{
-		SetSendReceiveState(EConcertSendReceiveState::SendOnly);
-	}
-
-	UE_DEPRECATED(5.2,"Use GetSendReceiveState to query the live update status.")
-	virtual bool IsSessionSuspended() const
-	{
-		return GetSendReceiveState() == EConcertSendReceiveState::SendOnly;
-	}
-
 	/** Get the send/receive state for this session. */
 	virtual EConcertSendReceiveState GetSendReceiveState() const = 0;
 
diff --git a/Engine/Plugins/Developer/Concert/ConcertMain/Source/ConcertTransport/Public/IConcertEndpoint.h b/Engine/Plugins/Developer/Concert/ConcertMain/Source/ConcertTransport/Public/IConcertEndpoint.h
--- a/Engine/Plugins/Developer/Concert/ConcertMain/Source/ConcertTransport/Public/IConcertEndpoint.h
+++ b/Engine/Plugins/Developer/Concert/ConcertMain/Source/ConcertTransport/Public/IConcertEndpoint.h
@@ -277,85 +277,31 @@ protected:
 		Response->SetRequestMessageId(RequestMessageId);
 	}
 
-	/** 
-	 * Add a Request Handler
-	 */
-	UE_DEPRECATED(5.1, "Types names are now represented by path names. Please use a version of this function that takes an FTopLevelAssetPath as MessageType.")
-	virtual void InternalAddRequestHandler(const FName& RequestMessageType, const TSharedRef<IConcertRequestHandler>& Handler)
-	{
-		InternalAddRequestHandler(UClass::TryConvertShortTypeNameToPathName<UStruct>(RequestMessageType.ToString()), Handler);
-	}
-
 	/**
 	 * Add a Request Handler
 	 */
 	virtual void InternalAddRequestHandler(const FTopLevelAssetPath& RequestMessageType, const TSharedRef<IConcertRequestHandler>& Handler) = 0;
 
-	/**
-	 * Remove an Request Handler
-	 */
-	UE_DEPRECATED(5.1, "Types names are now represented by path names. Please use a version of this function that takes an FTopLevelAssetPath as MessageType.")
-	virtual void InternalRemoveRequestHandler(const FName& RequestMessageType)
-	{
-		InternalRemoveRequestHandler(UClass::TryConvertShortTypeNameToPathName<UStruct>(RequestMessageType.ToString()));
-	}
-
 	/**
 	 * Remove an Request Handler
 	 */
 	virtual void InternalRemoveRequestHandler(const FTopLevelAssetPath& RequestMessageType) = 0;
 
-	/**
-	 * Add an Event Handler
-	 */
-	UE_DEPRECATED(5.1, "Types names are now represented by path names. Please use a version of this function that takes an FTopLevelAssetPath as MessageType.")
-	virtual void InternalAddEventHandler(const FName& EventMessageType, const TSharedRef<IConcertEventHandler>& Handler)
-	{
-		InternalAddEventHandler(UClass::TryConvertShortTypeNameToPathName<UStruct>(EventMessageType.ToString()), Handler);
-	}
-
 	/**
 	 * Add an Event Handler
 	 */
 	virtual void InternalAddEventHandler(const FTopLevelAssetPath& EventMessageType, const TSharedRef<IConcertEventHandler>& Handler) = 0;
 
-	/**
-	 * Remove an Event Handler
-	 */
-	UE_DEPRECATED(5.1, "Types names are now represented by path names. Please use a version of this function that takes an FTopLevelAssetPath as MessageType.")
-	virtual void InternalRemoveEventHandler(const FName& EventMessageType)
-	{
-		InternalRemoveEventHandler(UClass::TryConvertShortTypeNameToPathName<UStruct>(EventMessageType.ToString()));
-	}
-
 	/**
 	 * Remove an Event Handler
 	 */
 	virtual void InternalRemoveEventHandler(const FTopLevelAssetPath& EventMessageType) = 0;
 
-	/**
-	 * Subscribe to an Event
-	 */
-	UE_DEPRECATED(5.1, "Types names are now represented by path names. Please use a version of this function that takes an FTopLevelAssetPath as MessageType.")
-	virtual void InternalSubscribeToEvent(const FName& EventMessageType)
-	{
-		InternalSubscribeToEvent(UClass::TryConvertShortTypeNameToPathName<UStruct>(EventMessageType.ToString()));
-	}
-
 	/**
 	 * Subscribe to an Event
 	 */
 	virtual void InternalSubscribeToEvent(const FTopLevelAssetPath& EventMessageType) = 0;
 
-	/**
-	 * Unsubscribe from an Event
-	 */
-	UE_DEPRECATED(5.1, "Types names are now represented by path names. Please use a version of this function that takes an FTopLevelAssetPath as MessageType.")
-	virtual void InternalUnsubscribeFromEvent(const FName& EventMessageType)
-	{
-		InternalUnsubscribeFromEvent(UClass::TryConvertShortTypeNameToPathName<UStruct>(EventMessageType.ToString()));
-	}
-
 	/**
 	 * Unsubscribe from an Event
 	 */
