Allow patch versions to be compatible with Multi-user sessions.

diff --git a/Engine/Plugins/Developer/Concert/ConcertMain/Source/Concert/Private/ConcertVersion.cpp b/Engine/Plugins/Developer/Concert/ConcertMain/Source/Concert/Private/ConcertVersion.cpp
--- a/Engine/Plugins/Developer/Concert/ConcertMain/Source/Concert/Private/ConcertVersion.cpp
+++ b/Engine/Plugins/Developer/Concert/ConcertMain/Source/Concert/Private/ConcertVersion.cpp
@@ -7,6 +7,8 @@
 #include "Modules/BuildVersion.h"
 #include "Serialization/CustomVersion.h"
 
+#include "HAL/IConsoleManager.h"
+
 #include UE_INLINE_GENERATED_CPP_BY_NAME(ConcertVersion)
 
 LLM_DEFINE_TAG(Concert_ConcertVersion);
@@ -14,10 +16,29 @@ LLM_DEFINE_TAG(Concert_ConcertVersion);
 
 namespace ConcertVersionUtil
 {
+static TAutoConsoleVariable<int32> CVarVersionValidation(
+	TEXT("Concert.VersionValidationMode"), 0,
+	TEXT("Specifies if we are using default or legacy based version validation. In default mode (0), clients will be allowed to join sessions if patch versions are compatible. For example, 5.4.4 is version compatible with 5.4.3. In legacy mode (1), versions have to match exactly when joining a session."));
 
-bool ValidateVersion(const int32 InCurrent, const int32 InOther, const FText InVersionDisplayName, const EConcertVersionValidationMode InValidationMode, FText* OutFailureReason)
+bool ValidateVersion(const int32 InCurrent, const int32 InOther, const FText InVersionDisplayName, const EConcertVersionValidationMode InValidationMode, bool bIsVersionPatchValue, FText* OutFailureReason)
 {
-	switch (InValidationMode)
+	// Convert the InValidationMode into something we can handle.  If we are in default version mode then we need to convert patch compatible version to ::Identical or ::Compatible.
+	// 
+	const EConcertVersionValidationMode Mode = [bIsVersionPatchValue, InValidationMode]()
+	{
+		const bool bUseDefaultVersionValidation = CVarVersionValidation.GetValueOnAnyThread() == 0;
+		if (bUseDefaultVersionValidation && bIsVersionPatchValue && InValidationMode == EConcertVersionValidationMode::PatchCompatible)
+		{
+			return EConcertVersionValidationMode::Compatible;
+		}
+		else if (InValidationMode == EConcertVersionValidationMode::PatchCompatible)
+		{
+			return EConcertVersionValidationMode::Identical;
+		}
+		return InValidationMode;
+	}();
+	
+	switch (Mode)
 	{
 	case EConcertVersionValidationMode::Identical:
 		if (InOther != InCurrent)
@@ -62,9 +83,11 @@ void FConcertFileVersionInfo::Initialize()
 
 bool FConcertFileVersionInfo::Validate(const FConcertFileVersionInfo& InOther, const EConcertVersionValidationMode InValidationMode, FText* OutFailureReason) const
 {
-	return ConcertVersionUtil::ValidateVersion(FileVersion, InOther.FileVersion, LOCTEXT("UE4PackageVersionName", "UE4 Package Version"), InValidationMode, OutFailureReason)
-		&& ConcertVersionUtil::ValidateVersion(FileVersionUE5, InOther.FileVersionUE5, LOCTEXT("UE5PackageVersionName", "UE5 Package Version"), InValidationMode, OutFailureReason)
-		&& ConcertVersionUtil::ValidateVersion(FileVersionLicensee, InOther.FileVersionLicensee, LOCTEXT("LicenseePackageVersionName", "Licensee Package Version"), InValidationMode, OutFailureReason);
+	const bool bIsVersionPatchValue = false;
+	
+	return ConcertVersionUtil::ValidateVersion(FileVersion, InOther.FileVersion, LOCTEXT("UE4PackageVersionName", "UE4 Package Version"), InValidationMode, bIsVersionPatchValue, OutFailureReason)
+		&& ConcertVersionUtil::ValidateVersion(FileVersionUE5, InOther.FileVersionUE5, LOCTEXT("UE5PackageVersionName", "UE5 Package Version"), InValidationMode, bIsVersionPatchValue, OutFailureReason)
+		&& ConcertVersionUtil::ValidateVersion(FileVersionLicensee, InOther.FileVersionLicensee, LOCTEXT("LicenseePackageVersionName", "Licensee Package Version"), InValidationMode, bIsVersionPatchValue, OutFailureReason);
 }
 
 
@@ -78,10 +101,10 @@ void FConcertEngineVersionInfo::Initialize(const FEngineVersion& InVersion)
 
 bool FConcertEngineVersionInfo::Validate(const FConcertEngineVersionInfo& InOther, const EConcertVersionValidationMode InValidationMode, FText* OutFailureReason) const
 {
-	return ConcertVersionUtil::ValidateVersion(Major, InOther.Major, LOCTEXT("MajorEngineVersionName", "Major Engine Version"), InValidationMode, OutFailureReason)
-		&& ConcertVersionUtil::ValidateVersion(Minor, InOther.Minor, LOCTEXT("MinorEngineVersionName", "Minor Engine Version"), InValidationMode, OutFailureReason)
-		&& ConcertVersionUtil::ValidateVersion(Patch, InOther.Patch, LOCTEXT("PatchEngineVersionName", "Patch Engine Version"), InValidationMode, OutFailureReason)
-		&& ConcertVersionUtil::ValidateVersion(Changelist, InOther.Changelist, LOCTEXT("ChangelistEngineVersionName", "Changelist Engine Version"), InValidationMode, OutFailureReason);
+	return ConcertVersionUtil::ValidateVersion(Major, InOther.Major, LOCTEXT("MajorEngineVersionName", "Major Engine Version"), InValidationMode, false /*bIsVersionPatchValue*/, OutFailureReason)
+		&& ConcertVersionUtil::ValidateVersion(Minor, InOther.Minor, LOCTEXT("MinorEngineVersionName", "Minor Engine Version"), InValidationMode, false /*bIsVersionPatchValue*/, OutFailureReason)
+		&& ConcertVersionUtil::ValidateVersion(Patch, InOther.Patch, LOCTEXT("PatchEngineVersionName", "Patch Engine Version"), InValidationMode, true /*bIsVersionPatchValue*/, OutFailureReason)
+		&& ConcertVersionUtil::ValidateVersion(Changelist, InOther.Changelist, LOCTEXT("ChangelistEngineVersionName", "Changelist Engine Version"), InValidationMode, true /*bIsVersionPatchValue*/, OutFailureReason);
 }
 
 
@@ -95,7 +118,9 @@ void FConcertCustomVersionInfo::Initialize(const FCustomVersion& InVersion)
 bool FConcertCustomVersionInfo::Validate(const FConcertCustomVersionInfo& InOther, const EConcertVersionValidationMode InValidationMode, FText* OutFailureReason) const
 {
 	check(Key == InOther.Key);
-	return ConcertVersionUtil::ValidateVersion(Version, InOther.Version, FText::AsCultureInvariant(FriendlyName.IsNone() ? Key.ToString() : FriendlyName.ToString()), InValidationMode, OutFailureReason);
+	const bool bIsVersionPatchValue = false;
+	
+	return ConcertVersionUtil::ValidateVersion(Version, InOther.Version, FText::AsCultureInvariant(FriendlyName.IsNone() ? Key.ToString() : FriendlyName.ToString()), InValidationMode, bIsVersionPatchValue, OutFailureReason);
 }
 
 void FConcertSessionVersionInfo::Initialize(bool bSupportMixedBuildTypes)
diff --git a/Engine/Plugins/Developer/Concert/ConcertMain/Source/Concert/Public/ConcertVersion.h b/Engine/Plugins/Developer/Concert/ConcertMain/Source/Concert/Public/ConcertVersion.h
--- a/Engine/Plugins/Developer/Concert/ConcertMain/Source/Concert/Public/ConcertVersion.h
+++ b/Engine/Plugins/Developer/Concert/ConcertMain/Source/Concert/Public/ConcertVersion.h
@@ -15,6 +15,8 @@ enum EConcertVersionValidationMode : uint8
 	Identical,
 	/** Check that the version information is compatible (other version is equal or newer than the current version) */
 	Compatible,
+	/** Check that the major and minor version are identical but allow compatible patch versions.*/
+	PatchCompatible
 };
 
 /** Holds file version information */
diff --git a/Engine/Plugins/Developer/Concert/ConcertMain/Source/ConcertServer/Private/ConcertServer.cpp b/Engine/Plugins/Developer/Concert/ConcertMain/Source/ConcertServer/Private/ConcertServer.cpp
--- a/Engine/Plugins/Developer/Concert/ConcertMain/Source/ConcertServer/Private/ConcertServer.cpp
+++ b/Engine/Plugins/Developer/Concert/ConcertMain/Source/ConcertServer/Private/ConcertServer.cpp
@@ -1530,7 +1530,7 @@ bool FConcertServer::CanJoinSession(const TSharedPtr<IConcertServerSession>& Ser
 		return false;
 	}
 
-	if (ServerSession->GetSessionInfo().VersionInfos.Num() > 0 && !ServerSession->GetSessionInfo().VersionInfos.Last().Validate(SessionVersionInfo, EConcertVersionValidationMode::Identical, OutFailureReason))
+	if (ServerSession->GetSessionInfo().VersionInfos.Num() > 0 && !ServerSession->GetSessionInfo().VersionInfos.Last().Validate(SessionVersionInfo, EConcertVersionValidationMode::PatchCompatible, OutFailureReason))
 	{
 		return false;
 	}
