Fix disabled CONCERT_TRACE_INIT macro to match enabled.

#rnx

diff --git a/Engine/Plugins/Developer/Concert/ConcertMain/Source/ConcertTransport/Public/Trace/ConcertProtocolTrace.h b/Engine/Plugins/Developer/Concert/ConcertMain/Source/ConcertTransport/Public/Trace/ConcertProtocolTrace.h
--- a/Engine/Plugins/Developer/Concert/ConcertMain/Source/ConcertTransport/Public/Trace/ConcertProtocolTrace.h
+++ b/Engine/Plugins/Developer/Concert/ConcertMain/Source/ConcertTransport/Public/Trace/ConcertProtocolTrace.h
@@ -112,7 +112,7 @@ namespace UE::ConcertTrace
 #define CONCERT_TRACE_REPLICATION_OBJECT_SINK(SinkName, ObjectPath, SequenceId) CONCERT_TRACE_OBJECT_SINK(SinkName, ObjectPath, SequenceId, ProtocolSuiteToInt(UE::ConcertTrace::EProtocolSuite::Replication))
 
 #else
-#define CONCERT_TRACE_INIT(SessionId)
+#define CONCERT_TRACE_INIT(EndpointIdAsTOptional, InEndpointDisplayFString, IsServer)
 #define CONCERT_TRACE_OBJECT_SCOPE(EventName, ObjectPath, SequenceId, ProtocolId)
 #define CONCERT_TRACE_OBJECT_EVENT(EventName, ObjectPath, SequenceId, ProtocolId)
 #define CONCERT_TRACE_REPLICATION_OBJECT_SCOPE(EventName, ObjectPath, SequenceId)
