Move the discovery management off of the game thread so that we can remain in contact with a server even during long hitches.

diff --git a/Engine/Plugins/Developer/Concert/ConcertMain/Source/ConcertClient/Private/ConcertClient.cpp b/Engine/Plugins/Developer/Concert/ConcertMain/Source/ConcertClient/Private/ConcertClient.cpp
--- a/Engine/Plugins/Developer/Concert/ConcertMain/Source/ConcertClient/Private/ConcertClient.cpp
+++ b/Engine/Plugins/Developer/Concert/ConcertMain/Source/ConcertClient/Private/ConcertClient.cpp
@@ -16,8 +16,12 @@
 #include "Misc/CoreDelegates.h"
 #include "Misc/AsyncTaskNotification.h"
 #include "HAL/FileManager.h"
+#include "HAL/Runnable.h"
+#include "HAL/RunnableThread.h"
 #include "Stats/Stats.h"
 
+#include "Framework/Commands/UIAction.h"
+
 #include "Runtime/Launch/Resources/Version.h"
 
 #define LOCTEXT_NAMESPACE "ConcertClient"
@@ -651,6 +655,292 @@ FConcertClientPaths::FConcertClientPaths(const FString& InRole)
 {
 }
 
+
+namespace UE::ConcertClient::Private
+{
+
+TAutoConsoleVariable<int32> CVarLongHitchDuration(
+	TEXT("ConcertClient.Debug.GameThreadHitchDuration"),
+	30,
+	TEXT("How long to wait in seconds when ConcertClient.Debug.GameThreadHitch is called")
+);
+
+void ForceHitch()
+{
+	FPlatformProcess::Sleep(CVarLongHitchDuration.GetValueOnAnyThread());
+}
+	
+FAutoConsoleCommand StartTracingCommand(TEXT("ConcertClient.Debug.GameThreadHitch"),
+	TEXT("Force a long hitch on game thread."),
+	FExecuteAction::CreateStatic(&ForceHitch));
+
+}; // namespace UE::ConcertClient::Private
+
+class FConcertClientServerDiscoveryThread : public FRunnable
+{
+public:
+	explicit FConcertClientServerDiscoveryThread(TSharedPtr<IConcertLocalEndpoint> InClientAdminEndpoint, FString InRole, FString InAuthKey, int32 TimeoutInSeconds)
+		: bIsRunning(false)
+		, bStopRequested(false)
+		, ClientAdminEndpoint(MoveTemp(InClientAdminEndpoint))
+	{
+		Thread.Reset(FRunnableThread::Create(this, TEXT("ConcertClient Server Discovery Thread")));
+
+		DiscoveryEventData.CurrentRole = MoveTemp(InRole);
+		DiscoveryEventData.CurrentAuthKey = MoveTemp(InAuthKey);
+		DiscoveryEventData.DiscoveryTimeoutSeconds = TimeoutInSeconds;
+		ClientAdminEndpoint->RegisterEventHandler<FConcertAdmin_ServerDiscoveredEvent>(this, &FConcertClientServerDiscoveryThread::HandleServerDiscoveryEvent);
+	}
+
+	~FConcertClientServerDiscoveryThread()
+	{
+		if (Thread)
+		{
+			Thread->Kill(true);
+			Thread = nullptr;
+		}
+
+	}
+
+	void UnregisterHandler()
+	{
+		if (ClientAdminEndpoint.IsValid())
+		{
+			ClientAdminEndpoint->UnregisterEventHandler<FConcertAdmin_ServerDiscoveredEvent>();
+		}
+
+		ClientAdminEndpoint.Reset();
+	}
+		
+	bool IsRunning() const
+	{
+		return Thread && bIsRunning;
+	}
+
+	/**
+	 * Update the discovery information we send to the server. The user could change the timeout or
+	 * current authentication keys from the client settings. 
+	 **/
+	void SetDiscoveryEventInformation(FString InRole, FString InAuthKey, int32 TimeoutInSeconds)
+	{
+		FScopeLock DiscoveryDataLock(&DiscoveryEventDataCS);
+		DiscoveryEventData.CurrentRole = MoveTemp(InRole);
+		DiscoveryEventData.CurrentAuthKey = MoveTemp(InAuthKey);
+		DiscoveryEventData.DiscoveryTimeoutSeconds = TimeoutInSeconds;
+	}
+		
+	/**
+	 * For use by external callers to check to see if the server list has updated.  This is to allow the ticker
+	 * on the game thread to check for updated server list and broadcast out to existing listeners.  The original
+	 * implementation always existed on the game thread. This helps maintain that functionality and limit the
+	 * need for delegates to be thread safe.
+	 */
+	bool HasServerListChanged() const
+	{
+		return bServersUpdated;
+	}
+
+	/**
+	 * Reset the flag that tracks changes to the running server list.
+	 */
+	void ResetServerListChange()
+	{
+		bServersUpdated = false;
+	}
+	/**
+	 * Gets the current list of known servers.  We have to lock access to the known server list before returning values. 
+	 */
+	TArray<FConcertServerInfo> GetKnownServers()
+	{
+		FScopeLock KnownServersLock(&KnownServersCS);
+		TArray<FConcertServerInfo> ServerArray;
+		ServerArray.Empty(KnownServers.Num());
+		for (const auto& Server : KnownServers)
+		{
+			ServerArray.Emplace(Server.Value.ServerInfo);
+		}
+		return ServerArray;
+	}
+
+	//~ FRunnable Interface
+	virtual bool Init() override
+	{
+		bIsRunning = true;
+		return true;
+	}
+
+	/**
+	 * Every second send out a discovery message to clients and record discovered servers. Timeout any servers that are no
+	 * longer reporting online.
+	 * */
+	virtual uint32 Run() override
+	{
+		while (!bStopRequested)
+		{
+			const FDateTime UtcNow = FDateTime::UtcNow();
+			SendDiscoverServersEvent();
+			TimeoutDiscovery(UtcNow);
+			FPlatformProcess::SleepNoStats(1.0f);
+		}
+		return 0;
+	}
+
+	/** Stop the running thread and unregister our discovery handlers. */
+	virtual void Stop() override
+	{
+		UnregisterHandler();
+		bStopRequested = true;
+	}
+
+	/** Thread was told to exit by the FRunnable interface. */
+	virtual void Exit() override
+	{
+		bIsRunning = false;
+	}
+	//~ FRunnable Interface.
+
+private:
+	/**
+	 * Dirty the server updated flag that indicates that a new server has been discovered.
+	 */
+	void SetServersUpdatedFlag()
+	{
+		bServersUpdated = true;
+	}
+
+	/**
+	 * Handle a server discovery event and update our list of known Multi-user server.s
+	 */
+	void HandleServerDiscoveryEvent(const FConcertMessageContext& Context)
+	{
+		const FConcertAdmin_ServerDiscoveredEvent* Message = Context.GetMessage<FConcertAdmin_ServerDiscoveredEvent>();
+
+		if (Message->ConcertProtocolVersion != EConcertMessageVersion::LatestVersion)
+		{
+			return;
+		}
+
+		FScopeLock KnownServersLock(&KnownServersCS);
+		FKnownServer* Info = KnownServers.Find(Message->ConcertEndpointId);
+		if (Info == nullptr)
+		{
+			UE_LOG(LogConcert, Display, TEXT("Server %s discovered."), *Message->ServerName);
+			KnownServers.Emplace(Message->ConcertEndpointId, FKnownServer{ Context.UtcNow, FConcertServerInfo{ Message->ConcertEndpointId, Message->ServerName, Message->InstanceInfo, Message->ServerFlags } });
+			
+			SetServersUpdatedFlag();
+		}
+		else
+		{
+			Info->LastDiscoveryTime = Context.UtcNow;
+		}
+	}
+
+	/**
+	 * Send a discovery message to servers on the network. A server will send a response back which will be handled by
+	 * HandleServerDiscoveryEvent above.
+	 */
+	void SendDiscoverServersEvent()
+	{
+		FScopeLock DiscoveryDataLock(&DiscoveryEventDataCS);
+		FConcertAdmin_DiscoverServersEvent DiscoverServersEvent;
+		DiscoverServersEvent.ConcertProtocolVersion = EConcertMessageVersion::LatestVersion;
+		DiscoverServersEvent.RequiredRole = DiscoveryEventData.CurrentRole;;
+		DiscoverServersEvent.RequiredVersion = VERSION_STRINGIFY(ENGINE_MAJOR_VERSION) TEXT(".") VERSION_STRINGIFY(ENGINE_MINOR_VERSION);
+		DiscoverServersEvent.ClientAuthenticationKey = DiscoveryEventData.CurrentAuthKey;
+
+		/** This call ends up into Message Bus which is thread safe.  This will be called even if the main thread is blocked so that we do not lose contact with running servers. */
+		ClientAdminEndpoint->PublishEvent(DiscoverServersEvent);
+	}
+
+	/**
+	 * Timeout any stale servers that have stopped responding.
+	 */
+	void TimeoutDiscovery(const FDateTime& UtcNow)
+	{
+		int32 TimeoutInSeconds = 5;
+		{
+			FScopeLock DiscoveryDataLock(&DiscoveryEventDataCS);
+			TimeoutInSeconds = DiscoveryEventData.DiscoveryTimeoutSeconds;
+		}
+		
+		const FTimespan DiscoveryTimeoutSpan = FTimespan(0, 0, TimeoutInSeconds);
+		if (DiscoveryTimeoutSpan.IsZero())
+		{
+			return;
+		}
+
+		FScopeLock KnownServersLock(&KnownServersCS);
+		bool TimeoutOccurred = false;
+		for (auto It = KnownServers.CreateIterator(); It; ++It)
+		{
+			if (It->Value.LastDiscoveryTime + DiscoveryTimeoutSpan <= UtcNow)
+			{
+				TimeoutOccurred = true;
+				UE_LOG(LogConcert, Display, TEXT("Server %s lost."), *It->Value.ServerInfo.ServerName);
+				It.RemoveCurrent();
+				continue;
+			}
+		}
+	
+		if (TimeoutOccurred)
+		{
+			SetServersUpdatedFlag();
+		}
+	}
+
+	/** Simple struct to store the state of discovered servers.  Each server gets a last known time so that we can expire stale servers. */
+	struct FKnownServer
+	{
+		/** Last time we saw this server. */
+		FDateTime LastDiscoveryTime;
+
+		/** Server specific information: name, version, etc. */
+		FConcertServerInfo ServerInfo;
+	};
+
+	
+	/** Mutex to protect access to known servers map. */
+	mutable FCriticalSection KnownServersCS;
+
+	/** Map of discovered Concert Servers */
+	TMap<FGuid, FKnownServer> KnownServers;
+
+	/** Atomics for state tracking and external access to the runnable. */
+	std::atomic<bool> bIsRunning = false;
+	std::atomic<bool> bStopRequested = false;
+
+	/**
+	 * When this is true the servers list was updated.  The discovery ticker in ConcertClient will poll this periodically to determine if
+	 * a server updated delegate should be sent.
+	 */
+	std::atomic<bool> bServersUpdated = false;
+
+	/** Shared pointer access to the local endpoint for this client. */
+	TSharedPtr<IConcertLocalEndpoint> ClientAdminEndpoint;
+
+	/** Mutex for discovery event data below. */
+	mutable FCriticalSection DiscoveryEventDataCS;
+
+	/** Discovery event data is data that is required to send discovery events.  This data exists on the settings object but it is copied into the thread for message sending. */
+	struct FDiscoveryEventData
+	{
+		/** Populated and updated by the main thread during discovery. */
+		FString CurrentRole;
+
+		/** Authentication key for accessing the server. This is established at app startup. */
+		FString CurrentAuthKey;
+
+		/** The timespace at which discovered Multi-User server are considered stale if they haven't answered back. */
+		int32 DiscoveryTimeoutSeconds = 5;
+	};
+
+	/** Data we need to pass long in the discovery event message we send to the server. */
+	FDiscoveryEventData DiscoveryEventData;
+		
+	TUniquePtr<FRunnableThread> Thread;
+};
+
+
 FConcertClient::FConcertClient(const FString& InRole, const TSharedPtr<IConcertEndpointProvider>& InEndpointProvider)
 	: Role(InRole)
 	, Paths(InRole)
@@ -658,6 +948,7 @@ FConcertClient::FConcertClient(const FString& InRole, const TSharedPtr<IConcertE
 	, DiscoveryCount(0)
 	, bClientSessionPendingDestroy(false)
 {
+
 }
 
 FConcertClient::~FConcertClient()
@@ -741,7 +1032,6 @@ void FConcertClient::Shutdown()
 		StopDiscovery();
 	}
 	ClientAdminEndpoint.Reset();
-	KnownServers.Empty();
 
 	if (ClientSession.IsValid())
 	{
@@ -765,13 +1055,21 @@ void FConcertClient::StartDiscovery()
 	++DiscoveryCount;
 	if (ClientAdminEndpoint.IsValid() && !DiscoveryTick.IsValid())
 	{
-		ClientAdminEndpoint->RegisterEventHandler<FConcertAdmin_ServerDiscoveredEvent>(this, &FConcertClient::HandleServerDiscoveryEvent);
-
+		ServerDiscoveryThread = MakeUnique<FConcertClientServerDiscoveryThread>(ClientAdminEndpoint, Role, Settings->ClientSettings.ClientAuthenticationKey, Settings->ClientSettings.DiscoveryTimeoutSeconds);
+		
 		DiscoveryTick = FTSTicker::GetCoreTicker().AddTicker(TEXT("Discovery"), 1, [this](float DeltaSeconds) {
 			QUICK_SCOPE_CYCLE_COUNTER(STAT_FConcertClient_Discovery_Tick);
-			const FDateTime UtcNow = FDateTime::UtcNow();
-			SendDiscoverServersEvent();
-			TimeoutDiscovery(UtcNow);
+
+			if (ServerDiscoveryThread)
+			{
+				ServerDiscoveryThread->SetDiscoveryEventInformation(Role, Settings->ClientSettings.ClientAuthenticationKey, Settings->ClientSettings.DiscoveryTimeoutSeconds);
+				if (ServerDiscoveryThread->HasServerListChanged())
+				{
+					ServerDiscoveryThread->ResetServerListChange();
+					ServersUpdatedDelegate.Broadcast();
+				}
+			}
+
 			return true;
 		});
 	}
@@ -786,10 +1084,8 @@ void FConcertClient::StopDiscovery()
 		return;
 	}
 
-	if (ClientAdminEndpoint.IsValid())
-	{
-		ClientAdminEndpoint->UnregisterEventHandler<FConcertAdmin_ServerDiscoveredEvent>();
-	}
+	ServerDiscoveryThread->Stop();
+	
 	if (DiscoveryTick.IsValid())
 	{
 		FTSTicker::GetCoreTicker().RemoveTicker(DiscoveryTick);
@@ -841,13 +1137,11 @@ FConcertConnectionError FConcertClient::GetLastConnectionError() const
 
 TArray<FConcertServerInfo> FConcertClient::GetKnownServers() const
 {
-	TArray<FConcertServerInfo> ServerArray;
-	ServerArray.Empty(KnownServers.Num());
-	for (const auto& Server : KnownServers)
+	if (ServerDiscoveryThread)
 	{
-		ServerArray.Emplace(Server.Value.ServerInfo);
+		return ServerDiscoveryThread->GetKnownServers();
 	}
-	return ServerArray;
+	return {};
 }
 
 FSimpleMulticastDelegate& FConcertClient::OnKnownServersUpdated()
@@ -1339,64 +1633,6 @@ void FConcertClient::OnEndFrame()
 	}
 }
 
-void FConcertClient::TimeoutDiscovery(const FDateTime& UtcNow)
-{
-	const FTimespan DiscoveryTimeoutSpan = FTimespan(0, 0, Settings->ClientSettings.DiscoveryTimeoutSeconds);
-	if (DiscoveryTimeoutSpan.IsZero())
-	{
-		return;
-	}
-
-	bool TimeoutOccured = false;
-	for (auto It = KnownServers.CreateIterator(); It; ++It)
-	{
-		if (It->Value.LastDiscoveryTime + DiscoveryTimeoutSpan <= UtcNow)
-		{
-			TimeoutOccured = true;
-			UE_LOG(LogConcert, Display, TEXT("Server %s lost."), *It->Value.ServerInfo.ServerName);
-			It.RemoveCurrent();
-			continue;
-		}
-	}
-
-	if (TimeoutOccured)
-	{
-		ServersUpdatedDelegate.Broadcast();
-	}
-}
-
-void FConcertClient::SendDiscoverServersEvent()
-{
-	FConcertAdmin_DiscoverServersEvent DiscoverServersEvent;
-	DiscoverServersEvent.ConcertProtocolVersion = EConcertMessageVersion::LatestVersion;
-	DiscoverServersEvent.RequiredRole = Role;
-	DiscoverServersEvent.RequiredVersion = VERSION_STRINGIFY(ENGINE_MAJOR_VERSION) TEXT(".") VERSION_STRINGIFY(ENGINE_MINOR_VERSION);
-	DiscoverServersEvent.ClientAuthenticationKey = Settings->ClientSettings.ClientAuthenticationKey;
-	ClientAdminEndpoint->PublishEvent(DiscoverServersEvent);
-}
-
-void FConcertClient::HandleServerDiscoveryEvent(const FConcertMessageContext& Context)
-{
-	const FConcertAdmin_ServerDiscoveredEvent* Message = Context.GetMessage<FConcertAdmin_ServerDiscoveredEvent>();
-
-	if (Message->ConcertProtocolVersion != EConcertMessageVersion::LatestVersion)
-	{
-		return;
-	}
-
-	FKnownServer* Info = KnownServers.Find(Message->ConcertEndpointId);
-	if (Info == nullptr)
-	{
-		UE_LOG(LogConcert, Display, TEXT("Server %s discovered."), *Message->ServerName);
-		KnownServers.Emplace(Message->ConcertEndpointId, FKnownServer{ Context.UtcNow, FConcertServerInfo{ Message->ConcertEndpointId, Message->ServerName, Message->InstanceInfo, Message->ServerFlags } });
-		ServersUpdatedDelegate.Broadcast();
-	}
-	else
-	{
-		Info->LastDiscoveryTime = Context.UtcNow;
-	}
-}
-
 TFuture<EConcertResponseCode> FConcertClient::CreateClientSession(const FConcertSessionInfo& SessionInfo)
 {
 	check(SessionInfo.SessionId.IsValid() && !SessionInfo.SessionName.IsEmpty());
diff --git a/Engine/Plugins/Developer/Concert/ConcertMain/Source/ConcertClient/Private/ConcertClient.h b/Engine/Plugins/Developer/Concert/ConcertMain/Source/ConcertClient/Private/ConcertClient.h
--- a/Engine/Plugins/Developer/Concert/ConcertMain/Source/ConcertClient/Private/ConcertClient.h
+++ b/Engine/Plugins/Developer/Concert/ConcertMain/Source/ConcertClient/Private/ConcertClient.h
@@ -16,6 +16,7 @@ class FConcertPendingConnection;
 class FConcertClientJoinSessionTask;
 class FConcertClientCreateSessionTask;
 class FAsyncTaskNotification;
+class FConcertClientServerDiscoveryThread;
 
 class FConcertClientPaths
 {
@@ -122,15 +123,6 @@ private:
 	/** */
 	void OnEndFrame();
 
-	/** Remove server from the known server list when they haven't discovered them for a while */
-	void TimeoutDiscovery(const FDateTime& UtcNow);
-
-	/** Broadcast a message to discover Concert servers */
-	void SendDiscoverServersEvent();
-	
-	/** Handle any answers from Concert server to our search queries */
-	void HandleServerDiscoveryEvent(const FConcertMessageContext& Context);
-
 	/** Create a Concert client session based on the session information provided. The future is set when the client connection is established or has failed. */
 	TFuture<EConcertResponseCode> CreateClientSession(const FConcertSessionInfo& SessionInfo);
 
@@ -152,17 +144,12 @@ private:
 	/** Count of the number of times the discovery has been enabled */
 	uint32 DiscoveryCount;
 
-	/** Ticker for Discovering Concert Servers */
+	/** Ticker for interfacing with the discovery thread. */
 	FTSTicker::FDelegateHandle DiscoveryTick;
 
-	struct FKnownServer
-	{
-		FDateTime LastDiscoveryTime;
-		FConcertServerInfo ServerInfo;
-	};
-	/** Map of discovered Concert Servers */
-	TMap<FGuid, FKnownServer> KnownServers;
-
+	/** Unique pointer to the discovery thread runnable. */
+	TUniquePtr<FConcertClientServerDiscoveryThread> ServerDiscoveryThread;
+		
 	/** Holds a delegate to be invoked when the server list was updated. */
 	FSimpleMulticastDelegate ServersUpdatedDelegate;
 
diff --git a/Engine/Plugins/Developer/Concert/ConcertMain/Source/ConcertTransport/Private/ConcertLocalEndpoint.cpp b/Engine/Plugins/Developer/Concert/ConcertMain/Source/ConcertTransport/Private/ConcertLocalEndpoint.cpp
--- a/Engine/Plugins/Developer/Concert/ConcertMain/Source/ConcertTransport/Private/ConcertLocalEndpoint.cpp
+++ b/Engine/Plugins/Developer/Concert/ConcertMain/Source/ConcertTransport/Private/ConcertLocalEndpoint.cpp
@@ -669,6 +669,11 @@ void FConcertLocalEndpoint::ProcessEndpointDiscovery(const FConcertMessageContex
 			);
 		}
 	}
+	else
+	{
+		// We received a message from the remote endpoint, make sure we update our last receipt time. 
+		RemoteEndpoint->UpdateLastMessageReceivedTime(ConcertContext.UtcNow);
+	}
 }
 
 void FConcertLocalEndpoint::ProcessReliableHandshake(const FConcertMessageContext& ConcertContext)
