Check for a null discovery thread handler before calling Stop()

diff --git a/Engine/Plugins/Developer/Concert/ConcertMain/Source/ConcertClient/Private/ConcertClient.cpp b/Engine/Plugins/Developer/Concert/ConcertMain/Source/ConcertClient/Private/ConcertClient.cpp
--- a/Engine/Plugins/Developer/Concert/ConcertMain/Source/ConcertClient/Private/ConcertClient.cpp
+++ b/Engine/Plugins/Developer/Concert/ConcertMain/Source/ConcertClient/Private/ConcertClient.cpp
@@ -1084,7 +1084,10 @@ void FConcertClient::StopDiscovery()
 		return;
 	}
 
-	ServerDiscoveryThread->Stop();
+	if (ServerDiscoveryThread)
+	{
+		ServerDiscoveryThread->Stop();		
+	}
 	
 	if (DiscoveryTick.IsValid())
 	{
