Fix a crash in Multi-user server when users are not permitted to join because of an invalid session.

#rnx

diff --git a/Engine/Plugins/Developer/Concert/ConcertMain/Source/ConcertServer/Private/ConcertServer.cpp b/Engine/Plugins/Developer/Concert/ConcertMain/Source/ConcertServer/Private/ConcertServer.cpp
--- a/Engine/Plugins/Developer/Concert/ConcertMain/Source/ConcertServer/Private/ConcertServer.cpp
+++ b/Engine/Plugins/Developer/Concert/ConcertMain/Source/ConcertServer/Private/ConcertServer.cpp
@@ -1141,7 +1141,7 @@ TFuture<FConcertAdmin_SessionInfoResponse> FConcertServer::HandleFindSessionRequ
 	else
 	{
 		ResponseData.ResponseCode = EConcertResponseCode::Failed;
-		UE_LOG(LogConcert, Display, TEXT("Refusing user %s to join session %s (Id: %s, Owner: %s, Reason: %s)"), *Message->OwnerClientInfo.UserName, ServerSessionNamePtr, *Message->SessionId.ToString(), *ServerSession->GetSessionInfo().OwnerUserName, *ResponseData.Reason.ToString());
+		UE_LOG(LogConcert, Display, TEXT("Refusing user %s to join session %s (Id: %s, Owner: %s, Reason: %s)"), *Message->OwnerClientInfo.UserName, ServerSessionNamePtr, *Message->SessionId.ToString(), ServerSession ? *ServerSession->GetSessionInfo().OwnerUserName : TEXT("unknown owner"), *ResponseData.Reason.ToString());
 	}
 
 	return FConcertAdmin_SessionInfoResponse::AsFuture(MoveTemp(ResponseData));
