Implements a file sharing system that can be used with Multi-user. FConcertCloudSharingService will be created on server and client startup.  If users enable the service then a challenge will be made to the server to verify that the client and server can communicate over the file sharing service data channel.   If this fails at any point then we automatically fallback to the original behavior of sending files over the network.  Note that if a server starts a session with file sharing and one client cannot prove they can do file sharing then that client will not be allowed to join the session as the server will want to send file shares to clients once it is enabled on the server.

Note that file sharing only kicks in when a package size is greater than 2MB.  This is defined in ConcertWorkspaceData.h and carries forward the existing functionality that was used for disaster recovery.

The feature is considered experimental.

#Synced-CL 48909003

diff --git a/Engine/Plugins/Developer/Concert/ConcertMain/Source/Concert/Private/ConcertSettings.cpp b/Engine/Plugins/Developer/Concert/ConcertMain/Source/Concert/Private/ConcertSettings.cpp
--- a/Engine/Plugins/Developer/Concert/ConcertMain/Source/Concert/Private/ConcertSettings.cpp
+++ b/Engine/Plugins/Developer/Concert/ConcertMain/Source/Concert/Private/ConcertSettings.cpp
@@ -10,6 +10,28 @@
 namespace ConcertSettingsUtils
 {
 
+FText ValidateFileSharingPath(const FString& Path)
+{
+	if (Path.IsEmpty())
+	{
+		return LOCTEXT("PathName_Empty", "Path name is empty.");
+	}
+
+	FText Reason;
+	if (!FPaths::ValidatePath(Path, &Reason))
+	{
+		return FText::Format(LOCTEXT("IllegalPathName", "{0}"), Reason);
+	}
+
+	if (!FPaths::DirectoryExists(Path))
+	{
+		return FText::Format(LOCTEXT("DirectoryDoesNotExist", "Path {0} does not exist."), FText::FromString(Path));
+	}
+
+	return FText::GetEmpty();
+}
+	
+	
 /** Returns an error messages if the user display name is invalid, otherwise, returns an empty text. */
 FText ValidateDisplayName(const FString& Name)
 {
diff --git a/Engine/Plugins/Developer/Concert/ConcertMain/Source/Concert/Public/ConcertMessageData.h b/Engine/Plugins/Developer/Concert/ConcertMain/Source/Concert/Public/ConcertMessageData.h
--- a/Engine/Plugins/Developer/Concert/ConcertMain/Source/Concert/Public/ConcertMessageData.h
+++ b/Engine/Plugins/Developer/Concert/ConcertMain/Source/Concert/Public/ConcertMessageData.h
@@ -244,6 +244,63 @@ enum class EConcertPayloadSerializationMethod : uint8
 	Cbor,
 };
 
+USTRUCT()
+struct FConcertSessionChallengeData 
+{
+	GENERATED_BODY()
+
+	FConcertSessionChallengeData()
+	{
+		ChallengeKey = FGuid::NewGuid();
+		CreationTimeInTicks = FDateTime::Now().GetTicks();
+	}
+
+	bool operator==(const FConcertSessionChallengeData& Other) const = default;
+	
+	UPROPERTY()
+	FGuid ChallengeKey;
+
+	UPROPERTY()
+	int64 CreationTimeInTicks;
+};
+
+USTRUCT()
+struct FConcertSessionFileShareChallenge
+{
+	GENERATED_BODY()
+
+	UPROPERTY()
+	FString Path;
+
+	UPROPERTY()
+	FConcertSessionChallengeData Data;
+};
+
+USTRUCT()
+struct FConcertSessionFileShare
+{
+	GENERATED_BODY()
+
+	/**
+	 * Indicates that file sharing is enabled and available on the server.
+	 */
+	UPROPERTY()
+	bool bFileShareEnabled = false;
+
+	/**
+	 * Root path for all file shares for this session.
+	 */
+	UPROPERTY()
+	FString FileShareRoot;
+
+	/**
+	 * A challenge id that can be shared with the server / client to confirm that
+	 * file sharing can be established between a server and a client.
+	 */
+	UPROPERTY()
+	FConcertSessionFileShareChallenge Challenge;
+};
+
 /** Holds info on an instance communicating through concert */
 USTRUCT()
 struct FConcertInstanceInfo
@@ -421,6 +478,9 @@ struct FConcertSessionInfo
 	void SetLastModified(const FDateTime& Value) { LastModifiedTicks = Value.GetTicks(); }
 	void SetLastModifiedToNow() { SetLastModified(FDateTime::Now()); }
 
+	UPROPERTY()
+	FConcertSessionFileShare FileSharing;
+
 private:
 	
 	/**
diff --git a/Engine/Plugins/Developer/Concert/ConcertMain/Source/Concert/Public/ConcertMessages.h b/Engine/Plugins/Developer/Concert/ConcertMain/Source/Concert/Public/ConcertMessages.h
--- a/Engine/Plugins/Developer/Concert/ConcertMain/Source/Concert/Public/ConcertMessages.h
+++ b/Engine/Plugins/Developer/Concert/ConcertMain/Source/Concert/Public/ConcertMessages.h
@@ -272,6 +272,9 @@ struct FConcertAdmin_CreateSessionRequest : public FConcertRequestData
 
 	UPROPERTY(VisibleAnywhere, Category="Concert Message")
 	FConcertSessionVersionInfo VersionInfo;
+
+	UPROPERTY(VisibleAnywhere, Category="Concert Message")
+	FConcertSessionFileShare FileShare;
 };
 
 USTRUCT()
@@ -290,6 +293,9 @@ struct FConcertAdmin_FindSessionRequest : public FConcertRequestData
 
 	UPROPERTY(VisibleAnywhere, Category="Concert Message")
 	FConcertSessionVersionInfo VersionInfo;
+
+	UPROPERTY(VisibleAnywhere, Category="Concert Message")
+	FConcertSessionFileShare FileShare;
 };
 
 /** Used to copy a live session or restore an archived one. */
@@ -332,7 +338,10 @@ struct FConcertAdmin_SessionInfoResponse : public FConcertResponseData
 	GENERATED_BODY()
 
 	UPROPERTY(VisibleAnywhere, Category="Concert Message")
-	FConcertSessionInfo SessionInfo; // TODO: Split session Id out of session info
+	FConcertSessionInfo SessionInfo;
+
+	UPROPERTY(VisibleAnywhere, Category = "Concert Message")
+	FConcertSessionFileShare FileShare;
 };
 
 /** Create an archived copy of a live session. */
diff --git a/Engine/Plugins/Developer/Concert/ConcertMain/Source/Concert/Public/ConcertSettings.h b/Engine/Plugins/Developer/Concert/ConcertMain/Source/Concert/Public/ConcertSettings.h
--- a/Engine/Plugins/Developer/Concert/ConcertMain/Source/Concert/Public/ConcertSettings.h
+++ b/Engine/Plugins/Developer/Concert/ConcertMain/Source/Concert/Public/ConcertSettings.h
@@ -15,12 +15,15 @@
 namespace ConcertSettingsUtils
 {
 
-	/** Returns an error messages if the user display name is invalid, otherwise, returns an empty text. */
+/** Returns an error messages if the user display name is invalid, otherwise, returns an empty text. */
 FText CONCERT_API ValidateDisplayName(const FString& Name);
 
 /** Returns an error messages if the specified session name is invalid, otherwise, returns an empty text. */
 FText CONCERT_API ValidateSessionName(const FString& Name);
 
+/** Returns an error messages if the specified path name is invalid, otherwise, returns an empty text. */
+FText CONCERT_API ValidateFileSharingPath(const FString& Path);
+
 }
 USTRUCT()
 struct FConcertSessionSettings
diff --git a/Engine/Plugins/Developer/Concert/ConcertMain/Source/Concert/Public/IConcertFileSharingService.h b/Engine/Plugins/Developer/Concert/ConcertMain/Source/Concert/Public/IConcertFileSharingService.h
new file mode 100644
--- /dev/null
+++ b/Engine/Plugins/Developer/Concert/ConcertMain/Source/Concert/Public/IConcertFileSharingService.h
@@ -0,0 +1,98 @@
+// Copyright Epic Games, Inc. All Rights Reserved.
+
+#pragma once
+
+#include "Templates/SharedPointer.h"
+
+/**
+ * Implements a file sharing service used as a side channel to ConcertTransport to exchange large files. This API is designed to
+ * work hand in hand with the request/response/event layer.
+ *
+ * When a file needs to be exchanged between the client and the server:
+ *
+ * The publisher side would:
+ *    - Publish a file using file sharing service. The service tracks the request and capture the file for sharing (can sent it to a server asynchronously, can be copied to a shared folder, etc.)
+ *    - Notify the consumer(s) by sending a Concert request/event/response with the ID of the published file.
+ *
+ * The consumer(s) side would:
+ *    - Receive and process the Concert request/event/response and extract the file ID.
+ *    - Create a stream to read the shared file, ideally in a background task. If required, new method could be added to simplify asynchrous IO.
+ *
+ * * @note The lifetime of the shared files is implementation specific.
+ *
+ */
+class IConcertFileSharingService
+{
+public:
+	virtual ~IConcertFileSharingService() = default;
+
+	/** Make a clone of this sharing service to be used. Ownership is transferred to the caller. Callers can choose to wrap this into a unique or shared pointer. */
+	virtual IConcertFileSharingService* Clone() const = 0;
+
+	/** Indicates if the file sharing service is enabled. */
+	virtual bool IsEnabled() const = 0;
+
+	/** Turn on/off the file sharing service. */
+	virtual void SetEnabled(bool bInEnabled) = 0;
+
+	/**
+	 * Sets the desired path to use for the service. Implementors have the right to ignore the path specified.
+	 *
+	 * Implementors should set the enabled state based on if the path can be written. 
+	 * 
+	 * @param FileShareRoot to set as the root folder for content.
+	 */
+	virtual void SetPath(const FString& FileShareRoot) = 0;
+
+	/**
+	 * Get the assign path for file sharing root.
+	 */
+	virtual FString GetPath() const = 0;
+
+	/** Deletes the specified outfile id. */
+	virtual void Delete(const FString& InFileId) = 0;
+	
+	/**
+	 * Copy the specified files in a uniquely named file stored in the service directory for sharing. The lifetime of the shared file
+	 * is implementation specific.
+	 * @param Pathname The file to copy and share.
+	 * @param OutFileId The ID of the published file. The publisher is expected to send the ID through the Concert transport layer to the consumer(s) (request/response/event).
+	 * @return True if the file was published successfully.
+	 */
+	virtual bool Publish(const FString& Pathname, FString& OutFileId) = 0;
+
+	/**
+	 * Copies up to Count bytes contained in the archives in a uniquely named file stored in the service directory for sharing.
+	 * The functions reads the archive from its current position up to \a Count bytes. The lifetime of the shared file
+	 * is implementation specific.
+	 * @param SrcAr The Archive containing the bytes to share. The archive is already at the right position, not necessarily 0.
+	 * @param Count The number of bytes to read from the archive.
+	 * @param OutFileId The ID of the published file. The publisher is expected to send the ID through the Concert transport layer to the consumer(s) (request/response/event).
+	 * @return True if the file was published successfully.
+	 */
+	virtual bool Publish(FArchive& SrcAr, int64 Count, FString& OutFileId) = 0;
+
+	/**
+	 * Copies all the bytes contained in the archives in a uniquely named file stored in the service directory for sharing.
+	 * @param SrcAr The Archive containing the bytes to share. Expected to be at position 0.
+	 * @param OutFileId The ID of the published file. The publisher is expected to send the ID through the Concert transport layer to the consumer(s) (request/response/event).
+	 * @return True if the file was published successfully.
+	 */
+	bool Publish(FArchive& SrcAr, FString& OutFileId) { return Publish(SrcAr, SrcAr.TotalSize(), OutFileId); }
+
+	/**
+	 * Open a shared file for consumption.
+	 * @param InFileId The Id of the file to read. The consumer is expected to receives this ID through the Concert transport layer (request/response/event).
+	 * @return An archive (loading mode) to stream the content of the file or nullptr if the file could not be found or opened.
+	 * @note The function can be used to poll until the resources is fully available. A more sophisticated API could be added to wait for a ressource if needed later.
+	 */
+	virtual TSharedPtr<FArchive> CreateReader(const FString& InFileId) = 0;
+
+
+	/**
+	 * Open a file for writing.
+	 * @param OutFileId The id of the file opened for write. 
+	 * @return An archive (write mode) to stream the content of the file.
+	 */
+	virtual TSharedPtr<FArchive> CreateWriter(FString& OutFileId) = 0;
+};
diff --git a/Engine/Plugins/Developer/Concert/ConcertMain/Source/ConcertClient/Private/ConcertClient.cpp b/Engine/Plugins/Developer/Concert/ConcertMain/Source/ConcertClient/Private/ConcertClient.cpp
--- a/Engine/Plugins/Developer/Concert/ConcertMain/Source/ConcertClient/Private/ConcertClient.cpp
+++ b/Engine/Plugins/Developer/Concert/ConcertMain/Source/ConcertClient/Private/ConcertClient.cpp
@@ -8,6 +8,7 @@
 #include "ConcertLogGlobal.h"
 #include "ConcertTransportEvents.h"
 #include "ConcertClientSettings.h"
+
 #include "Algo/Transform.h"
 
 #include "Containers/Ticker.h"
@@ -600,6 +601,15 @@ public:
 				*ResolvedSessionName = FText::AsCultureInvariant(SessionInfoResponse.SessionInfo.SessionName);
 				if (SessionInfoResponse.ResponseCode == EConcertResponseCode::Success)
 				{
+					if (SessionInfoResponse.FileShare.bFileShareEnabled && Client->FileSharingService)
+					{
+						// Enable the file sharing service.
+						Client->FileSharingService->SetEnabled(true);
+
+						// Clean-up the challenge file we originally generated when we joined the session.
+						Client->FileSharingService->Delete(SessionInfoResponse.FileShare.Challenge.Path);
+					}
+
 					// If CreateClientSession() returns a failure, it is because the server did not reply to the 'join session' event and the endpoint timed out or the connection was canceled/aborted (for which there is a special message already).
 					ConnectionError = GetServerNotRespondingErrorMessage();
 					return Client->CreateClientSession(SessionInfoResponse.SessionInfo).Share();
@@ -636,6 +646,15 @@ public:
 				}
 				else if (SessionInfoResponse.ResponseCode == EConcertResponseCode::Success)
 				{
+					if (SessionInfoResponse.FileShare.bFileShareEnabled && Client->FileSharingService)
+					{
+						// Enable the file sharing service.
+						Client->FileSharingService->SetEnabled(true);
+
+						// Clean-up the challenge file we originally generated when we joined the session.
+						Client->FileSharingService->Delete(SessionInfoResponse.FileShare.Challenge.Path);
+					}
+
 					// If CreateClientSession() returns a failure, it is because the server did not reply to the 'join session' event and the endpoint timed out or the connection was canceled/aborted (for which there is a special message already).
 					ConnectionError = GetServerNotRespondingErrorMessage();
 					return Client->CreateClientSession(SessionInfoResponse.SessionInfo).Share();
@@ -674,6 +693,35 @@ FAutoConsoleCommand StartTracingCommand(TEXT("ConcertClient.Debug.GameThreadHitc
 	TEXT("Force a long hitch on game thread."),
 	FExecuteAction::CreateStatic(&ForceHitch));
 
+FConcertSessionFileShare MakeFileSharingChallenge(TSharedPtr<IConcertFileSharingService>& SharingService, const FDirectoryPath& FileShareRoot)
+{
+	FConcertSessionFileShare FileShare;
+	FileShare.bFileShareEnabled = false;
+	
+	// We should always have a file sharing service to work with.
+	if (SharingService)
+	{
+		FileShare.FileShareRoot = FileShareRoot.Path;
+		
+		SharingService->SetPath(FileShareRoot.Path);
+		
+		TSharedPtr<FArchive> Writer = SharingService->CreateWriter(FileShare.Challenge.Path);
+		if (Writer)
+		{
+			FileShare.bFileShareEnabled = true;
+			*Writer << FileShare.Challenge.Data.ChallengeKey;
+			*Writer << FileShare.Challenge.Data.CreationTimeInTicks;
+		}
+	}
+
+	if (!FileShare.bFileShareEnabled)
+	{
+		UE_LOG(LogConcert, Display, TEXT("File sharing is disabled this client."));
+	}
+
+	return FileShare;
+}
+	
 }; // namespace UE::ConcertClient::Private
 
 class FConcertClientServerDiscoveryThread : public FRunnable
@@ -1172,6 +1220,22 @@ FOnConcertClientSessionConnectionChanged& FConcertClient::OnSessionConnectionCha
 	return OnSessionConnectionChangedDelegate;
 }
 
+void FConcertClient::SetFileSharingService(TSharedPtr<IConcertFileSharingService> InFileSharingService)
+{
+	FileSharingService = MoveTemp(InFileSharingService);
+
+	if (FileSharingService)
+	{
+		// Start with the FileSharingService disabled.  Per session, we determine if it should be on.
+		FileSharingService->SetEnabled(false);
+	}
+}
+
+TSharedPtr<IConcertFileSharingService> FConcertClient::GetFileSharingService() const
+{
+	return FileSharingService;
+}
+
 EConcertConnectionStatus FConcertClient::GetSessionConnectionStatus() const
 {
 	return ClientSession.IsValid() ? ClientSession->GetConnectionStatus() : EConcertConnectionStatus::Disconnected;
@@ -1501,6 +1565,11 @@ TFuture<EConcertResponseCode> FConcertClient::InternalCreateSession(const FGuid&
 		CreateSessionRequest.SessionSettings.Initialize();
 		CreateSessionRequest.SessionSettings.ArchiveNameOverride = CreateSessionArgs.ArchiveNameOverride;
 
+		if (GetConfiguration()->ClientSettings.bEnableFileSharing)
+		{
+			CreateSessionRequest.FileShare = UE::ConcertClient::Private::MakeFileSharingChallenge(FileSharingService, GetConfiguration()->ClientSettings.CloudFileSharingPath);
+		}
+
 		ConnectionTasks.Emplace(MakeUnique<FConcertClientCreateSessionTask>(this, MoveTemp(CreateSessionRequest), ServerAdminEndpointId));
 	}
 
@@ -1546,6 +1615,11 @@ TFuture<EConcertResponseCode> FConcertClient::InternalJoinSession(const FGuid& S
 		// Session settings
 		FindSessionRequest.SessionSettings.Initialize();
 
+		if (GetConfiguration()->ClientSettings.bEnableFileSharing)
+		{
+			FindSessionRequest.FileShare = UE::ConcertClient::Private::MakeFileSharingChallenge(FileSharingService, GetConfiguration()->ClientSettings.CloudFileSharingPath);
+		}
+		
 		ConnectionTasks.Emplace(MakeUnique<FConcertClientJoinSessionTask>(this, MoveTemp(FindSessionRequest), ServerAdminEndpointId, ResolvedSessionName));
 	}
 
@@ -1617,6 +1691,12 @@ void FConcertClient::InternalDisconnectSession()
 		OnSessionShutdownDelegate.Broadcast(ClientSession.ToSharedRef());
 		ClientSession->Shutdown();
 		ClientSession.Reset();
+
+		// Turn off the file sharing service.
+		if (FileSharingService)
+		{
+			FileSharingService->SetEnabled(false);
+		}
 	}
 
 	bClientSessionPendingDestroy = false;
@@ -1642,18 +1722,16 @@ TFuture<EConcertResponseCode> FConcertClient::CreateClientSession(const FConcert
 
 	InternalDisconnectSession();
 	ClientSession = MakeShared<FConcertClientSession>(
-		SessionInfo, 
-		ClientInfo, 
-		Settings->ClientSettings, 
-		EndpointProvider->CreateLocalEndpoint(SessionInfo.SessionName, Settings->EndpointSettings, [this](const FConcertEndpointContext& Context)
-		{
+		SessionInfo,
+		ClientInfo,
+		Settings->ClientSettings,
+		EndpointProvider->CreateLocalEndpoint(SessionInfo.SessionName, Settings->EndpointSettings, [this](const FConcertEndpointContext& Context) {
 			return ConcertUtil::CreateLogger(Context, [this](const FConcertLog& Log)
-			{
-				ConcertTransportEvents::OnConcertClientLogEvent().Broadcast(*this, Log);
+			{ ConcertTransportEvents::OnConcertClientLogEvent().Broadcast(*this, Log);
 			});
 		}),
-		Paths.GetSessionWorkingDir(SessionInfo.SessionId)
-		);
+		Paths.GetSessionWorkingDir(SessionInfo.SessionId));
+
 	OnSessionStartupDelegate.Broadcast(ClientSession.ToSharedRef());
 	ClientSession->OnConnectionChanged().AddRaw(this, &FConcertClient::HandleSessionConnectionChanged);
 	ClientSession->Startup();
diff --git a/Engine/Plugins/Developer/Concert/ConcertMain/Source/ConcertClient/Private/ConcertClient.h b/Engine/Plugins/Developer/Concert/ConcertMain/Source/ConcertClient/Private/ConcertClient.h
--- a/Engine/Plugins/Developer/Concert/ConcertMain/Source/ConcertClient/Private/ConcertClient.h
+++ b/Engine/Plugins/Developer/Concert/ConcertMain/Source/ConcertClient/Private/ConcertClient.h
@@ -76,7 +76,9 @@ public:
 
 	virtual FOnConcertClientSessionGetPreConnectionTasks& OnGetPreConnectionTasks() override;
 	virtual FOnConcertClientSessionConnectionChanged& OnSessionConnectionChanged() override;
-
+	virtual void SetFileSharingService(TSharedPtr<IConcertFileSharingService> InFileSharingService) override;
+	virtual TSharedPtr<IConcertFileSharingService> GetFileSharingService() const override;		
+	
 	virtual EConcertConnectionStatus GetSessionConnectionStatus() const override;
 	virtual TFuture<EConcertResponseCode> CreateSession(const FGuid& ServerAdminEndpointId, const FConcertCreateSessionArgs& CreateSessionArgs) override;
 	virtual TFuture<EConcertResponseCode> JoinSession(const FGuid& ServerAdminEndpointId, const FGuid& SessionId) override;
@@ -141,6 +143,9 @@ private:
 	/** Administration endpoint for the client (i.e. creating, joining sessions) */
 	TSharedPtr<IConcertLocalEndpoint> ClientAdminEndpoint;
 
+	/** Shared pointer to the file sharing service to use.  */
+	TSharedPtr<IConcertFileSharingService> FileSharingService;
+
 	/** Count of the number of times the discovery has been enabled */
 	uint32 DiscoveryCount;
 
diff --git a/Engine/Plugins/Developer/Concert/ConcertMain/Source/ConcertClient/Public/ConcertClientSettings.h b/Engine/Plugins/Developer/Concert/ConcertMain/Source/ConcertClient/Public/ConcertClientSettings.h
--- a/Engine/Plugins/Developer/Concert/ConcertMain/Source/ConcertClient/Public/ConcertClientSettings.h
+++ b/Engine/Plugins/Developer/Concert/ConcertMain/Source/ConcertClient/Public/ConcertClientSettings.h
@@ -85,6 +85,26 @@ struct FConcertClientSettings
 	/** A key used to identify the clients during server discovery. If the server was configured to restrict access, the client key must be know of the server. Can be left empty. */
 	UPROPERTY(config)
 	FString ClientAuthenticationKey;
+
+	/**
+	 * Specifies if we should use a file sharing service for larger assets that are saved in a Multi-user session. Note the specified folder must be
+	 * visible to all clients.
+	 */
+	UPROPERTY(Config, EditAnywhere, Category = "File Sharing")
+	bool bEnableFileSharing = false;
+
+	/**
+	 * Indicates that the file sharing content should be removed on exit. If this is false then it will be left in place and must be removed by the user.
+	 */
+	UPROPERTY(Config, EditAnywhere, meta=(EditCondition="bEnableFileSharing"), Category = "File Sharing")
+	bool bCleanSessionFilesOnDisconnect = true;
+
+	/**
+	 * The path for that shares files across clients. This path must be visible to all nodes. A challenge request will be sent to the server to confirm that it can see the path.
+	 * If the server cannot see the path then it will reject file sharing for this node and it will fallback to a non-fileshare code flow. 
+	 */
+	UPROPERTY(Config, EditAnywhere, meta=(EditCondition="bEnableFileSharing"), Category = "File Sharing")
+	FDirectoryPath CloudFileSharingPath;
 };
 
 UENUM()
diff --git a/Engine/Plugins/Developer/Concert/ConcertMain/Source/ConcertClient/Public/IConcertClient.h b/Engine/Plugins/Developer/Concert/ConcertMain/Source/ConcertClient/Public/IConcertClient.h
--- a/Engine/Plugins/Developer/Concert/ConcertMain/Source/ConcertClient/Public/IConcertClient.h
+++ b/Engine/Plugins/Developer/Concert/ConcertMain/Source/ConcertClient/Public/IConcertClient.h
@@ -8,6 +8,7 @@
 #include "IConcertSession.h"
 #include "ConcertMessages.h"
 #include "ConcertTransportMessages.h"
+#include "IConcertFileSharingService.h"
 
 class UConcertClientConfig;
 
@@ -280,6 +281,17 @@ public:
 	 */
 	virtual FOnConcertClientSessionConnectionChanged& OnSessionConnectionChanged() = 0;
 
+	/**
+	 * Sets the handle to the file sharing service.  The client needs to prove to the server that they both can see the
+	 * service.  This file sharing service will ultimate get promoted to the workspace to be used during a working session.
+	 */
+	virtual void SetFileSharingService(TSharedPtr<IConcertFileSharingService> InFileSharingService) = 0;
+
+	/**
+	 * Get the current file sharing service held by the client. 
+	 */
+	virtual TSharedPtr<IConcertFileSharingService> GetFileSharingService() const = 0;		
+
 	/**
 	 * Get the connection status of client session or disconnected if no session is present
 	 * @see EConcertConnectionStatus
diff --git a/Engine/Plugins/Developer/Concert/ConcertMain/Source/ConcertServer/Private/ConcertServer.cpp b/Engine/Plugins/Developer/Concert/ConcertMain/Source/ConcertServer/Private/ConcertServer.cpp
--- a/Engine/Plugins/Developer/Concert/ConcertMain/Source/ConcertServer/Private/ConcertServer.cpp
+++ b/Engine/Plugins/Developer/Concert/ConcertMain/Source/ConcertServer/Private/ConcertServer.cpp
@@ -10,6 +10,7 @@
 #include "ConcertLogGlobal.h"
 #include "ConcertTransportEvents.h"
 #include "IConcertServerEventSink.h"
+#include "IConcertFileSharingService.h"
 
 #include "Algo/AnyOf.h"
 #include "Misc/App.h"
@@ -119,6 +120,12 @@ void FConcertServer::Configure(const UConcertServerConfig* InSettings)
 		ServerInfo.ServerFlags |= EConcertServerFlags::IgnoreSessionRequirement;
 	}
 
+	if (Settings->bEnableFileSharing && FileSharingService)
+	{
+		FileSharingService->SetEnabled(true);
+		FileSharingService->SetPath(Settings->RootFileSharingPath.Path);
+	}
+
 	SessionRepositoryRootDir = FPaths::ProjectSavedDir() / Role / TEXT("Sessions"); // Server default session repository root dir.
 	if (!Settings->SessionRepositoryRootDir.IsEmpty())
 	{
@@ -149,6 +156,21 @@ const FConcertServerInfo& FConcertServer::GetServerInfo() const
 	return ServerInfo;
 }
 
+void FConcertServer::SetFileSharingService(TSharedPtr<IConcertFileSharingService> InFileSharingService)
+{
+	FileSharingService = MoveTemp(InFileSharingService);
+	if (FileSharingService && Settings && Settings->bEnableFileSharing)
+	{
+		FileSharingService->SetEnabled(true);
+		FileSharingService->SetPath(Settings->RootFileSharingPath.Path);
+	}
+}
+
+TSharedPtr<IConcertFileSharingService> FConcertServer::GetFileSharingService() const
+{
+	return FileSharingService;
+}
+
 TArray<FConcertEndpointContext> FConcertServer::GetRemoteAdminEndpoints() const
 {
 	if (IsStarted())
@@ -1088,10 +1110,93 @@ TFuture<FConcertAdmin_DropSessionRepositoriesResponse> FConcertServer::HandleDro
 	return FConcertAdmin_DropSessionRepositoriesResponse::AsFuture(MoveTemp(ResponseData));
 }
 
+
+EConcertResponseCode FConcertServer::ValidateFileSharing(const FConcertSessionFileShare& InRequest, FConcertSessionFileShare& OutFileShare, FText& OutFailureReason)
+{
+	// Assume that we can't support file sharing and prove otherwise.
+	OutFileShare.bFileShareEnabled = false;
+	OutFileShare.FileShareRoot = TEXT("");
+
+	/* Supply any file share information back to the client. */
+	if (FileSharingService && InRequest.bFileShareEnabled && Settings->bEnableFileSharing)
+	{
+		// Try to accept the file sharing challenge
+		OutFileShare.Challenge = InRequest.Challenge;
+
+		FString BaseDir = FPaths::ConvertRelativePathToFull(Settings->RootFileSharingPath.Path);
+		FString Normalized = FPaths::ConvertRelativePathToFull(InRequest.Challenge.Path);
+
+		if (!BaseDir.IsEmpty() && FPaths::IsUnderDirectory(Normalized, BaseDir))
+		{
+			TSharedPtr<FArchive> Reader = FileSharingService->CreateReader(Normalized);
+			if (Reader)
+			{
+				FConcertSessionChallengeData ChallengeData;
+				*Reader << ChallengeData.ChallengeKey;
+				*Reader << ChallengeData.CreationTimeInTicks;
+	
+				if (ChallengeData == InRequest.Challenge.Data)
+				{
+					// We can accept the challenge.
+					OutFileShare.bFileShareEnabled = true;
+					OutFileShare.FileShareRoot = Settings->RootFileSharingPath.Path;
+				}
+			}	
+		}
+
+		if (!OutFileShare.bFileShareEnabled)
+		{
+			OutFailureReason = LOCTEXT("Error_FailedChallenge", "The client failed to prove they could support file sharing.");
+			return EConcertResponseCode::Failed;
+		}
+	}
+	else if (FileSharingService && FileSharingService->IsEnabled() && Settings->bEnableFileSharing)
+	{
+		OutFailureReason = LOCTEXT("Error_NoFileSharing", "The server is setup for file sharing. Clients joining a session must support file sharing.");
+		return EConcertResponseCode::Failed;
+	}
+	else if (InRequest.bFileShareEnabled)
+	{
+		OutFailureReason = LOCTEXT("Error_FileSharingFromClient", "The client asked for file sharing but the server does not support it.");		
+		return EConcertResponseCode::Failed;
+	}
+
+	return EConcertResponseCode::Success;
+}
+
+template<typename TCreateOrJoinMessage>
+bool FConcertServer::FillSessionResponseInfoFromFileShareChallenge(const TCreateOrJoinMessage* Message, FConcertAdmin_SessionInfoResponse& ResponseData, const FConcertSessionInfo& SessionInfo)
+{
+	const TCHAR* ServerSessionNamePtr = *SessionInfo.SessionName;
+
+	if (Message->FileShare.bFileShareEnabled)
+	{
+		UE_LOG(LogConcert, Display, TEXT("Validating file sharing challenge from user %s to join session %s (Id: %s, Owner: %s)"), *Message->OwnerClientInfo.UserName, ServerSessionNamePtr, *SessionInfo.SessionId.ToString(), *SessionInfo.OwnerUserName);
+	}
+	
+	ResponseData.ResponseCode = ValidateFileSharing(Message->FileShare, ResponseData.FileShare, ResponseData.Reason);
+
+	if (ResponseData.ResponseCode != EConcertResponseCode::Success)
+	{
+		UE_LOG(LogConcert, Display, TEXT("Refusing user %s to join session %s because of failed file share challenge. (Id: %s, Owner: %s, Reason: %s)"), *Message->OwnerClientInfo.UserName, ServerSessionNamePtr, *SessionInfo.SessionId.ToString(), *SessionInfo.OwnerUserName, *ResponseData.Reason.ToString());
+	}
+	else
+	{
+		// Fill out the session info to return back to the client.
+		ResponseData.SessionInfo = SessionInfo;
+		UE_LOG(LogConcert, Display, TEXT("Allowing user %s to join session %s.  File sharing is %s. (Id: %s, Owner: %s)"), *Message->OwnerClientInfo.UserName, ServerSessionNamePtr, ResponseData.FileShare.bFileShareEnabled ? TEXT("enabled") : TEXT("disabled"), *SessionInfo.SessionId.ToString(), *SessionInfo.OwnerUserName);
+	}
+
+	return ResponseData.ResponseCode == EConcertResponseCode::Success;  
+}
+
 TFuture<FConcertAdmin_SessionInfoResponse> FConcertServer::HandleCreateSessionRequest(const FConcertMessageContext& Context)
 {
 	const FConcertAdmin_CreateSessionRequest* Message = Context.GetMessage<FConcertAdmin_CreateSessionRequest>();
 
+	// We have a valid session if it succeeded
+	FConcertAdmin_SessionInfoResponse ResponseData;
+
 	// Create a new server session
 	FText CreateFailureReason;
 	TSharedPtr<IConcertServerSession> NewServerSession;
@@ -1103,17 +1208,19 @@ TFuture<FConcertAdmin_SessionInfoResponse> FConcertServer::HandleCreateSessionRe
 		SessionInfo.SessionName = Message->SessionName;
 		SessionInfo.Settings = Message->SessionSettings;
 		SessionInfo.VersionInfos.Add(Message->VersionInfo);
-		NewServerSession = CreateSession(SessionInfo, CreateFailureReason);
+		if (FillSessionResponseInfoFromFileShareChallenge(Message, ResponseData, SessionInfo))
+		{
+			// Copy the session file sharing information to the user. 
+			SessionInfo.FileSharing = ResponseData.FileShare;
+			NewServerSession = CreateSession(SessionInfo, CreateFailureReason);
+			if (NewServerSession)
+			{
+				ResponseData.SessionInfo = NewServerSession->GetSessionInfo();
+			}
+		}
 	}
 
-	// We have a valid session if it succeeded
-	FConcertAdmin_SessionInfoResponse ResponseData;
-	if (NewServerSession)
-	{
-		ResponseData.SessionInfo = NewServerSession->GetSessionInfo();
-		ResponseData.ResponseCode = EConcertResponseCode::Success;
-	}
-	else
+	if (!NewServerSession)
 	{
 		ResponseData.ResponseCode = EConcertResponseCode::Failed;
 		ResponseData.Reason = CreateFailureReason;
@@ -1131,15 +1238,13 @@ TFuture<FConcertAdmin_SessionInfoResponse> FConcertServer::HandleFindSessionRequ
 
 	// Find the session requested
 	TSharedPtr<IConcertServerSession> ServerSession = GetLiveSession(Message->SessionId);
-	const TCHAR* ServerSessionNamePtr = ServerSession ? *ServerSession->GetName() : TEXT("<unknown>");
 	if (CanJoinSession(ServerSession, Message->SessionSettings, Message->VersionInfo, Message->ConcertEndpointId, Message->OwnerClientInfo, &ResponseData.Reason))
 	{
-		ResponseData.ResponseCode = EConcertResponseCode::Success;
-		ResponseData.SessionInfo = ServerSession->GetSessionInfo();
-		UE_LOG(LogConcert, Display, TEXT("Allowing user %s to join session %s (Id: %s, Owner: %s)"), *Message->OwnerClientInfo.UserName, ServerSessionNamePtr, *Message->SessionId.ToString(), *ServerSession->GetSessionInfo().OwnerUserName);
+		FillSessionResponseInfoFromFileShareChallenge(Message, ResponseData, ServerSession->GetSessionInfo());
 	}
 	else
 	{
+		const TCHAR* ServerSessionNamePtr = ServerSession ? *ServerSession->GetName() : TEXT("<unknown>");
 		ResponseData.ResponseCode = EConcertResponseCode::Failed;
 		UE_LOG(LogConcert, Display, TEXT("Refusing user %s to join session %s (Id: %s, Owner: %s, Reason: %s)"), *Message->OwnerClientInfo.UserName, ServerSessionNamePtr, *Message->SessionId.ToString(), ServerSession ? *ServerSession->GetSessionInfo().OwnerUserName : TEXT("unknown owner"), *ResponseData.Reason.ToString());
 	}
diff --git a/Engine/Plugins/Developer/Concert/ConcertMain/Source/ConcertServer/Private/ConcertServer.h b/Engine/Plugins/Developer/Concert/ConcertMain/Source/ConcertServer/Private/ConcertServer.h
--- a/Engine/Plugins/Developer/Concert/ConcertMain/Source/ConcertServer/Private/ConcertServer.h
+++ b/Engine/Plugins/Developer/Concert/ConcertMain/Source/ConcertServer/Private/ConcertServer.h
@@ -26,6 +26,9 @@ public:
 	virtual const UConcertServerConfig* GetConfiguration() const override;
 	virtual const FConcertServerInfo& GetServerInfo() const override;
 
+	virtual void SetFileSharingService(TSharedPtr<IConcertFileSharingService> InFileSharingService) override;
+	virtual TSharedPtr<IConcertFileSharingService> GetFileSharingService() const override;		
+
 	virtual TArray<FConcertEndpointContext> GetRemoteAdminEndpoints() const override;
 	virtual FOnConcertRemoteEndpointConnectionChanged& OnRemoteEndpointConnectionChanged() override;
 	virtual FMessageAddress GetRemoteAddress(const FGuid& AdminEndpointId) const override;
@@ -168,6 +171,13 @@ private:
 	/**  */
 	TSharedPtr<IConcertServerSession> RestoreArchivedSession(const FGuid& ArchivedSessionId, const FConcertSessionInfo& NewSessionInfo, const FConcertSessionFilter& SessionFilter, FText& OutFailureReason);
 
+	/** Validate a file a session join request for file sharing support. */
+	EConcertResponseCode ValidateFileSharing(const FConcertSessionFileShare& InRequest, FConcertSessionFileShare& OutFileShare, FText& OutFailureReason);
+
+	/** Helper template method that handles the two different message types that are used for joining a session (create / join existing). */
+	template<typename TCreateOrJoinMessage>
+	bool FillSessionResponseInfoFromFileShareChallenge(const TCreateOrJoinMessage* Message, FConcertAdmin_SessionInfoResponse& ResponseData, const FConcertSessionInfo& SessionInfo);
+		
 	/** The role of this server (eg, MultiUser, DisasterRecovery, etc) */
 	FString Role;
 
@@ -192,6 +202,9 @@ private:
 	/** Administration endpoint for the server (i.e. creating, joining sessions) */
 	TSharedPtr<IConcertLocalEndpoint> ServerAdminEndpoint;
 
+	/** File sharing service that we use to verify reachability to client proof files. */
+	TSharedPtr<IConcertFileSharingService> FileSharingService;
+		
 	/** Called when ServerAdminEndpoint emits IConcertLocalEndpoint::OnConcertRemoteEndpointConnectionChanged */
 	FOnConcertRemoteEndpointConnectionChanged OnConcertRemoteEndpointConnectionChangedDelegate; 
 	
@@ -210,5 +223,5 @@ private:
 	/** The root directory containing the session repositories. */
 	FString SessionRepositoryRootDir;
 
-	FOnConcertMessageAcknowledgementReceivedFromLocalEndpoint OnConcertMessageAcknowledgementReceivedFromLocalEndpoint; 
+	FOnConcertMessageAcknowledgementReceivedFromLocalEndpoint OnConcertMessageAcknowledgementReceivedFromLocalEndpoint;
 };
diff --git a/Engine/Plugins/Developer/Concert/ConcertMain/Source/ConcertServer/Public/ConcertServerSettings.h b/Engine/Plugins/Developer/Concert/ConcertMain/Source/ConcertServer/Public/ConcertServerSettings.h
--- a/Engine/Plugins/Developer/Concert/ConcertMain/Source/ConcertServer/Public/ConcertServerSettings.h
+++ b/Engine/Plugins/Developer/Concert/ConcertMain/Source/ConcertServer/Public/ConcertServerSettings.h
@@ -124,6 +124,23 @@ public:
 	/** If neither of WorkingDir and ArchiveDir are set, determine whether the server should mount a standard default session repository where new session will be created. */
 	UPROPERTY(config)
 	bool bMountDefaultSessionRepository = true;
+
+	/**
+	 * Specifies if we should use a file sharing service for larger assets that are saved in a Multi-user session.
+	 * Note the specified folder must be visible to all clients.
+	 *
+	 * This can be set on the command line using -CONCERTFILESHARINGENABLED
+	 */
+	UPROPERTY(config, EditAnywhere, Category = "File Sharing")
+	bool bEnableFileSharing = false;
+
+	/**
+	 * The path for that shares files across clients. This path must be visible to all nodes. 
+	 *
+	 * This can be set on the command line using -CONCERTFILESHAREPATH=<FullPathName>
+	 */
+	UPROPERTY(config, EditAnywhere, meta=(EditCondition="bEnableFileSharing"), Category = "File Sharing")
+	FDirectoryPath RootFileSharingPath;
 };
 
 #undef UE_API
diff --git a/Engine/Plugins/Developer/Concert/ConcertMain/Source/ConcertServer/Public/IConcertServer.h b/Engine/Plugins/Developer/Concert/ConcertMain/Source/ConcertServer/Public/IConcertServer.h
--- a/Engine/Plugins/Developer/Concert/ConcertMain/Source/ConcertServer/Public/IConcertServer.h
+++ b/Engine/Plugins/Developer/Concert/ConcertMain/Source/ConcertServer/Public/IConcertServer.h
@@ -7,6 +7,7 @@
 #include "IConcertSession.h"
 
 class UConcertServerConfig;
+class IConcertFileSharingService;
 struct FConcertClientInfo;
 
 DECLARE_MULTICAST_DELEGATE_OneParam(FOnConcertServerSessionStartup, TWeakPtr<IConcertServerSession>);
@@ -44,7 +45,17 @@ public:
 	 */
 	virtual const FConcertServerInfo& GetServerInfo() const = 0;
 
-	
+	/**
+	 * Sets the handle to the file sharing service.  The server needs to verify that it can access the challenge provided by
+	 * the client to confirm they both have access to the file sharing service. 
+	 */
+	virtual void SetFileSharingService(TSharedPtr<IConcertFileSharingService> InFileSharingService) = 0;
+
+	/**
+	 * Get the current file sharing service held by the server. 
+	 */
+	virtual TSharedPtr<IConcertFileSharingService> GetFileSharingService() const = 0;		
+
 	/**
 	 * Gets all remote admin endpoint IDs.
 	 * A remote admin endpoint is not connected to any session and communicating with the server otherwise, e.g. discovering sessions.
diff --git a/Engine/Plugins/Developer/Concert/ConcertSync/ConcertSyncClient/Source/ConcertSyncClient/Private/ConcertClientPackageManager.cpp b/Engine/Plugins/Developer/Concert/ConcertSync/ConcertSyncClient/Source/ConcertSyncClient/Private/ConcertClientPackageManager.cpp
--- a/Engine/Plugins/Developer/Concert/ConcertSync/ConcertSyncClient/Source/ConcertSyncClient/Private/ConcertClientPackageManager.cpp
+++ b/Engine/Plugins/Developer/Concert/ConcertSync/ConcertSyncClient/Source/ConcertSyncClient/Private/ConcertClientPackageManager.cpp
@@ -522,7 +522,7 @@ void FConcertClientPackageManager::HandleLocalPackageEvent(const FConcertPackage
 					return;
 				}
 			}
-			else if (FileSharingService && EnumHasAnyFlags(LiveSession->GetSessionFlags(), EConcertSyncSessionFlags::EnableFileSharing))
+			else if (IsFileSharingEnabled())
 			{
 				// Publish a copy of the package data in the sharing service and set the corresponding file ID in the event.
 				if (!FileSharingService->Publish(PackagePathname, PackageEvent.Package.FileId))
@@ -662,9 +662,14 @@ void FConcertClientPackageManager::PurgePendingPackages()
 	}
 }
 
+bool FConcertClientPackageManager::IsFileSharingEnabled() const
+{
+	return FileSharingService && FileSharingService->IsEnabled() && EnumHasAnyFlags(LiveSession->GetSessionFlags(), EConcertSyncSessionFlags::EnableFileSharing);
+}
+
 bool FConcertClientPackageManager::CanExchangePackageDataAsByteArray(int64 PackageDataSize) const
 {
-	if (FileSharingService && EnumHasAnyFlags(LiveSession->GetSessionFlags(), EConcertSyncSessionFlags::EnableFileSharing))
+	if (IsFileSharingEnabled())
 	{
 		return FConcertPackage::ShouldEmbedPackageDataAsByteArray(PackageDataSize); // Test the package data size against a preferred limit.
 	}
diff --git a/Engine/Plugins/Developer/Concert/ConcertSync/ConcertSyncClient/Source/ConcertSyncClient/Private/ConcertClientPackageManager.h b/Engine/Plugins/Developer/Concert/ConcertSync/ConcertSyncClient/Source/ConcertSyncClient/Private/ConcertClientPackageManager.h
--- a/Engine/Plugins/Developer/Concert/ConcertSync/ConcertSyncClient/Source/ConcertSyncClient/Private/ConcertClientPackageManager.h
+++ b/Engine/Plugins/Developer/Concert/ConcertSync/ConcertSyncClient/Source/ConcertSyncClient/Private/ConcertClientPackageManager.h
@@ -114,7 +114,11 @@ public:
 	bool IsReloadingPackage(FName PackageName) const;
 
 private:
-
+	/**
+	 * Is file sharing enabled.
+	 */
+	bool IsFileSharingEnabled() const;
+		
 	/**
 	 * Returns a full path for given package name if it was deleted and exists on the non-sandbox area.
 	 */
diff --git a/Engine/Plugins/Developer/Concert/ConcertSync/ConcertSyncClient/Source/ConcertSyncClient/Private/ConcertClientWorkspace.cpp b/Engine/Plugins/Developer/Concert/ConcertSync/ConcertSyncClient/Source/ConcertSyncClient/Private/ConcertClientWorkspace.cpp
--- a/Engine/Plugins/Developer/Concert/ConcertSync/ConcertSyncClient/Source/ConcertSyncClient/Private/ConcertClientWorkspace.cpp
+++ b/Engine/Plugins/Developer/Concert/ConcertSync/ConcertSyncClient/Source/ConcertSyncClient/Private/ConcertClientWorkspace.cpp
@@ -664,6 +664,19 @@ void FConcertClientWorkspace::HandleConnectionChanged(IConcertClientSession& InS
 
 	if (Status == EConcertConnectionStatus::Connected)
 	{
+		if (FileSharingService)
+		{
+			const FConcertSessionInfo& SessionInfo = InSession.GetSessionInfo();
+			FileSharingService->SetEnabled(SessionInfo.FileSharing.bFileShareEnabled);
+			if (SessionInfo.FileSharing.bFileShareEnabled)
+			{
+				FString WorkspacePath = SessionInfo.FileSharing.FileShareRoot / SessionInfo.SessionName / LiveSession->GetSession().GetSessionClientEndpointId().ToString();
+				FileSharingService->SetPath(WorkspacePath);
+			}
+
+			UE_LOG(LogConcert, Display, TEXT("File sharing is %s for session workspace."), FileSharingService->IsEnabled() ? TEXT("enabled") : TEXT("disabled") );
+		}
+		
 		bHasSyncedWorkspace = false;
 		bFinalizeWorkspaceSyncRequested = false;
 		InitialSyncSlowTask = MakeUnique<FScopedSlowTask>(1.0f, LOCTEXT("SynchronizingSession", "Synchronizing Session..."));
@@ -714,6 +727,15 @@ void FConcertClientWorkspace::HandleConnectionChanged(IConcertClientSession& InS
 		FConcertSlowTaskStackWorkaround::Get().PopTask(MoveTemp(InitialSyncSlowTask));
 		UE::ConcertWorkspace::Private::LevelInstanceForceEditorWorldMode.Reset();
 		UE::ConcertWorkspace::Private::EditorPaths.Reset();
+
+		if (FileSharingService && FileSharingService->IsEnabled())
+		{
+			IConcertClientRef ConcertClient = OwnerSyncClient->GetConcertClient();
+			if (ConcertClient->GetConfiguration()->ClientSettings.bCleanSessionFilesOnDisconnect)
+			{
+				FileSharingService->Delete(FileSharingService->GetPath());
+			}
+		}
 	}
 }
 
diff --git a/Engine/Plugins/Developer/Concert/ConcertSync/ConcertSyncClient/Source/ConcertSyncClient/Private/ConcertSyncClient.cpp b/Engine/Plugins/Developer/Concert/ConcertSync/ConcertSyncClient/Source/ConcertSyncClient/Private/ConcertSyncClient.cpp
--- a/Engine/Plugins/Developer/Concert/ConcertSync/ConcertSyncClient/Source/ConcertSyncClient/Private/ConcertSyncClient.cpp
+++ b/Engine/Plugins/Developer/Concert/ConcertSync/ConcertSyncClient/Source/ConcertSyncClient/Private/ConcertSyncClient.cpp
@@ -8,10 +8,13 @@
 #include "ConcertClientSettings.h"
 #include "ConcertSyncClientLiveSession.h"
 #include "ConcertClientWorkspace.h"
+
 #include "ConcertClientSequencerManager.h"
 #include "ConcertClientPresenceManager.h"
 #include "Replication/Manager/ReplicationManager.h"
 
+#include "Misc/Paths.h"
+
 LLM_DECLARE_TAG(Concert_ConcertSyncClient);
 #define LOCTEXT_NAMESPACE "ConcertSyncClient"
 
@@ -42,8 +45,7 @@ void FConcertSyncClient::Startup(const UConcertClientConfig* InClientConfig, con
 {
 	LLM_SCOPE_BYTAG(Concert_ConcertSyncClient);
 	SessionFlags = InSessionFlags;
-
-	// Boot the client instance
+	
 	ConcertClient->Configure(InClientConfig);
 	ConcertClient->Startup();
 	Bridges.TransactionBridge->StartBridge();
@@ -147,14 +149,13 @@ void FConcertSyncClient::GetSessionClientActions(const FConcertSessionClientInfo
 
 void FConcertSyncClient::SetFileSharingService(TSharedPtr<IConcertFileSharingService> InFileSharingService)
 {
-	check(!FileSharingService); // Not really meant to be set more than once.
-	FileSharingService = MoveTemp(InFileSharingService);
+	ConcertClient->SetFileSharingService(MoveTemp(InFileSharingService));
 }
 
 void FConcertSyncClient::CreateWorkspace(const TSharedRef<FConcertSyncClientLiveSession>& InLiveSession)
 {
 	DestroyWorkspace();
-	Workspace = MakeShared<FConcertClientWorkspace>(UE::ConcertSyncClient::FSessionBindArgs{ InLiveSession, Bridges } , FileSharingService, this);
+	Workspace = MakeShared<FConcertClientWorkspace>(UE::ConcertSyncClient::FSessionBindArgs{ InLiveSession, Bridges } , ConcertClient->GetFileSharingService(), this);
 	OnWorkspaceStartupDelegate.Broadcast(Workspace);
 #if WITH_EDITOR
 	SequencerManager->SetActiveWorkspace(Workspace);
diff --git a/Engine/Plugins/Developer/Concert/ConcertSync/ConcertSyncClient/Source/ConcertSyncClient/Private/ConcertSyncClient.h b/Engine/Plugins/Developer/Concert/ConcertSync/ConcertSyncClient/Source/ConcertSyncClient/Private/ConcertSyncClient.h
--- a/Engine/Plugins/Developer/Concert/ConcertSync/ConcertSyncClient/Source/ConcertSyncClient/Private/ConcertSyncClient.h
+++ b/Engine/Plugins/Developer/Concert/ConcertSync/ConcertSyncClient/Source/ConcertSyncClient/Private/ConcertSyncClient.h
@@ -86,9 +86,6 @@ private:
 	
 	/** Delegate called just after the workspace was deleted. */
 	FOnConcertClientSyncSessionStartupOrShutdown OnSyncSessionShutdownDelegate;
-
-	/** Optional side channel to exchange large blobs (package data) with the server in a scalable way (ex. the request/response transport layer is not designed and doesn't support exchanging 3GB packages). */
-	TSharedPtr<IConcertFileSharingService> FileSharingService;
 	
 	/** Manages the client's replication system. Can be used to join & leave replication sessions. */
 	TUniquePtr<UE::ConcertSyncClient::Replication::FReplicationManager> ReplicationManager;
diff --git a/Engine/Plugins/Developer/Concert/ConcertSync/ConcertSyncClient/Source/ConcertSyncClient/Private/ConcertSyncClientModule.cpp b/Engine/Plugins/Developer/Concert/ConcertSync/ConcertSyncClient/Source/ConcertSyncClient/Private/ConcertSyncClientModule.cpp
--- a/Engine/Plugins/Developer/Concert/ConcertSync/ConcertSyncClient/Source/ConcertSyncClient/Private/ConcertSyncClientModule.cpp
+++ b/Engine/Plugins/Developer/Concert/ConcertSync/ConcertSyncClient/Source/ConcertSyncClient/Private/ConcertSyncClientModule.cpp
@@ -6,6 +6,9 @@
 #include "ConcertSettings.h"
 #include "ConcertClientSettings.h"
 #include "ConcertLogGlobal.h"
+
+#include "ConcertCloudFileSharingService.h"
+
 #include "ConcertClientPackageBridge.h"
 #include "ConcertClientTransactionBridge.h"
 #include "Replication/ConcertClientReplicationBridge.h"
@@ -72,7 +75,8 @@ public:
 			OverwriteConfigParamIfValid(&ConcertSettingsUtils::ValidateSessionName, TEXT("-CONCERTSESSIONTORESTORE="), ClientConfig->DefaultSessionToRestore);
 			OverwriteConfigParamIfValid(&ConcertSettingsUtils::ValidateSessionName, TEXT("-CONCERTSAVESESSIONAS="), ClientConfig->DefaultSaveSessionAs);
 			OverwriteConfigParamIfValid(&ConcertSettingsUtils::ValidateDisplayName, TEXT("-CONCERTDISPLAYNAME="), ClientConfig->ClientSettings.DisplayName);
-
+			OverwriteConfigParamIfValid(&ConcertSettingsUtils::ValidateFileSharingPath, TEXT("-CONCERTFILESHAREPATH="), ClientConfig->ClientSettings.CloudFileSharingPath.Path);
+				
 			ClientConfig->bAutoConnect |= FParse::Param(CommandLine, TEXT("CONCERTAUTOCONNECT"));
 			FParse::Bool(CommandLine, TEXT("-CONCERTAUTOCONNECT="), ClientConfig->bAutoConnect);
 
@@ -89,6 +93,14 @@ public:
 			ClientConfig->EndpointSettings.bEnableLogging |= FParse::Param(CommandLine, TEXT("CONCERTLOGGING"));
 			FParse::Bool(CommandLine, TEXT("-CONCERTLOGGING="), ClientConfig->EndpointSettings.bEnableLogging);
 
+			ClientConfig->ClientSettings.bEnableFileSharing |= FParse::Param(CommandLine, TEXT("CONCERTFILESHARINGENABLED"));
+			FParse::Bool(CommandLine, TEXT("-CONCERTFILESHARINGENABLED="), ClientConfig->ClientSettings.bEnableFileSharing);
+
+			if (ClientConfig->ClientSettings.CloudFileSharingPath.Path.IsEmpty())
+			{
+				ClientConfig->ClientSettings.bEnableFileSharing = false;
+			}
+
 			// CONCERTTAGS
 			{
 				FString CmdTags;
@@ -105,7 +117,6 @@ public:
 				}
 			}
 		}
-
 		return ClientConfig;
 	}
 
@@ -119,9 +130,10 @@ public:
 		const TSharedRef<IConcertSyncClient> NewClient = MakeShared<FConcertSyncClient>(
 			InRole,
 			UE::ConcertSyncClient::FConcertBridges{ PackageBridge.Get(),
-			TransactionBridge.Get(),
-			ReplicationBridge.Get() }
-			);
+				TransactionBridge.Get(),
+				ReplicationBridge.Get() });
+
+		NewClient->SetFileSharingService(MakeShared<FConcertCloudFileSharingService>(InRole));
 		Clients.Add(NewClient);
 		OnClientCreated().Broadcast(NewClient);
 
diff --git a/Engine/Plugins/Developer/Concert/ConcertSync/ConcertSyncCore/Source/ConcertSyncCore/Private/ConcertCloudFileSharingService.cpp b/Engine/Plugins/Developer/Concert/ConcertSync/ConcertSyncCore/Source/ConcertSyncCore/Private/ConcertCloudFileSharingService.cpp
new file mode 100644
--- /dev/null
+++ b/Engine/Plugins/Developer/Concert/ConcertSync/ConcertSyncCore/Source/ConcertSyncCore/Private/ConcertCloudFileSharingService.cpp
@@ -0,0 +1,178 @@
+// Copyright Epic Games, Inc. All Rights Reserved.
+
+#include "ConcertCloudFileSharingService.h"
+
+#include "HAL/FileManager.h"
+#include "ConcertUtil.h"
+#include "Misc/Guid.h"
+#include "Serialization/JsonSerializer.h"
+#include "Misc/Paths.h"
+#include "Misc/PackageName.h"
+#include "HAL/PlatformProcess.h"
+
+#include "ConcertLogGlobal.h"
+
+FConcertCloudFileSharingService::FConcertCloudFileSharingService(const FString& Role)
+{
+}
+
+FConcertCloudFileSharingService::~FConcertCloudFileSharingService()
+{
+}
+
+IConcertFileSharingService* FConcertCloudFileSharingService::Clone() const
+{
+	return new FConcertCloudFileSharingService(*this);
+}
+
+bool FConcertCloudFileSharingService::IsEnabled() const
+{
+	return bEnabled && PathIsValid();
+}
+
+void FConcertCloudFileSharingService::SetEnabled(bool bInEnabled)
+{
+	bEnabled = bInEnabled;
+}
+
+bool FConcertCloudFileSharingService::PathIsValid() const
+{
+	FFileStatData StatData = IPlatformFile::GetPlatformPhysical().GetStatData(*SharedRootPathname);
+	return !StatData.bIsReadOnly && StatData.bIsValid && StatData.bIsDirectory;
+}
+
+void FConcertCloudFileSharingService::SetPath(const FString& InPath)
+{
+	FString Path = FPaths::ConvertRelativePathToFull(InPath);
+	FFileStatData StatData = IPlatformFile::GetPlatformPhysical().GetStatData(*Path);
+	if (!StatData.bIsReadOnly)
+	{
+		if (!IFileManager::Get().DirectoryExists(*Path))
+		{
+			IPlatformFile::GetPlatformPhysical().CreateDirectoryTree(*Path);
+		}
+	}
+
+	// Even though we may not have a valid path, we still store the path name for use.
+	SharedRootPathname = Path; 
+
+	if (!PathIsValid())
+	{
+		UE_LOG(LogConcert, Display, TEXT("ConcertCloudFileSharingService - Unable to set path to %s. File sharing will be disabled."), *Path);
+	}
+}
+
+FString FConcertCloudFileSharingService::GetPath() const
+{
+	return SharedRootPathname;
+}
+
+void FConcertCloudFileSharingService::Delete(const FString& InFileOrDirectory)
+{
+	FString FileOrDirectory = FPaths::ConvertRelativePathToFull(InFileOrDirectory);
+	if (!(FPaths::IsUnderDirectory(FileOrDirectory, SharedRootPathname) || FPaths::IsSamePath(FileOrDirectory, SharedRootPathname)))
+	{
+		return;
+	}
+
+	FFileStatData StatData = IPlatformFile::GetPlatformPhysical().GetStatData(*FileOrDirectory);
+	if (StatData.bIsValid && StatData.bIsDirectory)
+	{
+		IPlatformFile::GetPlatformPhysical().DeleteDirectoryRecursively(*FileOrDirectory);
+	}
+	else if (StatData.bIsValid)
+	{
+		IFileManager::Get().Delete(*FileOrDirectory);
+	}
+}
+
+bool FConcertCloudFileSharingService::Publish(const FString& Pathname, FString& OutFileId)
+{
+	// The file share must be enabled to allow for publishing.
+	if (!IsEnabled())
+	{
+		UE_LOG(LogConcert, Error, TEXT("ConcertCloudFileSharingService - Publishing %s is not allowed because sharing."), *Pathname);
+		return false;
+	}
+
+	FString PackageName;
+	if (FPackageName::TryConvertFilenameToLongPackageName(Pathname, PackageName))
+	{
+		FString Path = FPackageName::GetLongPackagePath(PackageName);
+		FString OutPath = SharedRootPathname / Path;
+		if (!IFileManager::Get().DirectoryExists(*OutPath))
+		{
+			IPlatformFile::GetPlatformPhysical().CreateDirectoryTree(*OutPath);
+		}
+
+		const bool bIncludeDot = true;
+		FString Extension = FPaths::GetExtension( Pathname, true );
+
+		// Copy File will fail if it already exists. If the user saves a new version of an asset, we simply append a version number to
+		// keep uniqueness.
+		FString PathNoExt = SharedRootPathname / PackageName;
+		FString CandidatePath = PathNoExt + Extension;
+		int32 VersionIdx = 2;
+		while (IFileManager::Get().FileExists(*CandidatePath))
+		{
+			CandidatePath = PathNoExt + TEXT("_v") + FString::FromInt(VersionIdx) + Extension;
+			VersionIdx++;
+		}
+
+		OutFileId = CandidatePath;
+	}
+	else
+	{
+		OutFileId = SharedRootPathname / FGuid::NewGuid().ToString();
+	}
+		
+	// Create a unique name for the file.
+	UE_LOG(LogConcert, Display, TEXT("ConcertCloudFileSharingService - Publishing %s to %s."), *Pathname, *OutFileId);
+	return IFileManager::Get().Copy(*OutFileId, *Pathname) == COPY_OK;
+}
+
+bool FConcertCloudFileSharingService::Publish(FArchive& SrcAr, int64 Size, FString& OutFileId)
+{
+	// The file share must be enabled tollow for publishing.
+	if (!IsEnabled())
+	{
+		UE_LOG(LogConcert, Error, TEXT("ConcertCloudFileSharingService - Publishing archive is not allowed because sharing is not enabled."));
+		return false;
+	}
+
+	// Create a unique name for the file.
+	FString SharedFilePathname = SharedRootPathname / FGuid::NewGuid().ToString();
+	TUniquePtr<FArchive> SharedFileWriter(IFileManager::Get().CreateFileWriter(*SharedFilePathname, FILEWRITE_AllowRead));
+	if (SharedFileWriter)
+	{
+		ConcertUtil::Copy(*SharedFileWriter, SrcAr, Size);
+		OutFileId = FPaths::ConvertRelativePathToFull(SharedFilePathname);
+		UE_LOG(LogConcert, Display, TEXT("ConcertCloudFileSharingService - Publishing Archive to %s."), *OutFileId);
+		return true;
+	}
+	return false;
+}
+
+TSharedPtr<FArchive> FConcertCloudFileSharingService::CreateReader(const FString& InFileId)
+{
+	FString FullInFileId = FPaths::ConvertRelativePathToFull(InFileId);
+	if (FPaths::IsUnderDirectory(FullInFileId, SharedRootPathname) && IFileManager::Get().FileExists(*FullInFileId))
+	{
+		UE_LOG(LogConcert, Verbose, TEXT("ConcertCloudFileSharingService - Creating reader for %s."), *FullInFileId);
+		return MakeShareable(IFileManager::Get().CreateFileReader(*FullInFileId));
+	}
+
+	return nullptr; // File not found.
+}
+
+TSharedPtr<FArchive> FConcertCloudFileSharingService::CreateWriter(FString& OutFileId)
+{
+	// Users are allow to call this function for generating challenge files. These files can be
+	if (PathIsValid())
+	{
+		OutFileId = SharedRootPathname / FGuid::NewGuid().ToString();
+		return MakeShareable(IFileManager::Get().CreateFileWriter(*OutFileId, FILEWRITE_AllowRead));
+	}
+
+	return nullptr;
+}
diff --git a/Engine/Plugins/Developer/Concert/ConcertSync/ConcertSyncCore/Source/ConcertSyncCore/Private/ConcertLocalFileSharingService.cpp b/Engine/Plugins/Developer/Concert/ConcertSync/ConcertSyncCore/Source/ConcertSyncCore/Private/ConcertLocalFileSharingService.cpp
--- a/Engine/Plugins/Developer/Concert/ConcertSync/ConcertSyncCore/Source/ConcertSyncCore/Private/ConcertLocalFileSharingService.cpp
+++ b/Engine/Plugins/Developer/Concert/ConcertSync/ConcertSyncCore/Source/ConcertSyncCore/Private/ConcertLocalFileSharingService.cpp
@@ -8,6 +8,7 @@
 #include "Misc/Paths.h"
 #include "HAL/PlatformProcess.h"
 
+#include "ConcertLogGlobal.h"
 
 FConcertLocalFileSharingService::FConcertLocalFileSharingService(const FString& Role)
 	: SharedRootPathname(FPaths::EngineIntermediateDir() / Role / TEXT("TempFileShare")) // The EngineIntermediateDir is expected to be know of both the client and the server.
@@ -34,6 +35,38 @@ FConcertLocalFileSharingService::~FConcertLocalFileSharingService()
 	SaveActiveServices(ActivePids);
 }
 
+IConcertFileSharingService* FConcertLocalFileSharingService::Clone() const
+{
+	return new FConcertLocalFileSharingService(*this);
+}
+
+bool FConcertLocalFileSharingService::IsEnabled() const
+{
+	// If the user has specified a local file sharing service then we are always enabled.  
+	return true;
+}
+
+void FConcertLocalFileSharingService::SetEnabled(bool bInEnabled)
+{
+	// We do not support disabling this service. If the user has specified a local file service then we always use it.
+}
+
+void FConcertLocalFileSharingService::SetPath(const FString& Path)
+{
+	// We do not support setting althernative paths with this file service. 
+}
+
+FString FConcertLocalFileSharingService::GetPath() const
+{
+	return SharedRootPathname;
+}
+
+void FConcertLocalFileSharingService::Delete(const FString& OutFileId)
+{
+	// This class does not handle the delete method.
+	UE_LOG(LogConcert, Warning, TEXT("LocalFileSharingService - Delete is not supported with this file sharing service. Leaving %s inplace."), *OutFileId);
+}
+
 void FConcertLocalFileSharingService::RemoveDeadProcessesAndFiles(TArray<uint32>& InOutPids)
 {
 	// Remove all dead process.
@@ -99,6 +132,12 @@ TSharedPtr<FArchive> FConcertLocalFileSharingService::CreateReader(const FString
 	return nullptr; // File not found.
 }
 
+TSharedPtr<FArchive> FConcertLocalFileSharingService::CreateWriter(FString& OutFileId)
+{
+	OutFileId = SharedRootPathname / FGuid::NewGuid().ToString();
+	return MakeShareable(IFileManager::Get().CreateFileWriter(*OutFileId, FILEWRITE_AllowRead));
+}
+
 void FConcertLocalFileSharingService::LoadActiveServices(TArray<uint32>& OutPids)
 {
 	if (IFileManager::Get().FileExists(*ActiveServicesRepositoryPathname))
diff --git a/Engine/Plugins/Developer/Concert/ConcertSync/ConcertSyncCore/Source/ConcertSyncCore/Public/ConcertCloudFileSharingService.h b/Engine/Plugins/Developer/Concert/ConcertSync/ConcertSyncCore/Source/ConcertSyncCore/Public/ConcertCloudFileSharingService.h
new file mode 100644
--- /dev/null
+++ b/Engine/Plugins/Developer/Concert/ConcertSync/ConcertSyncCore/Source/ConcertSyncCore/Public/ConcertCloudFileSharingService.h
@@ -0,0 +1,64 @@
+// Copyright Epic Games, Inc. All Rights Reserved.
+
+#pragma once
+
+#include "IConcertFileSharingService.h"
+
+#define UE_API CONCERTSYNCCORE_API
+
+/**
+ * Share files using a path that is connected to a cloud based shared file service, e.g. a shared drive. It will work across all clients and server
+ * assuming that the drive is visible. Validation checks are run during session connection to establish that both the server and client can see the
+ * same system.
+  */
+class FConcertCloudFileSharingService : public IConcertFileSharingService
+{
+public:
+	/**
+	 * Construct the cloud file sharing service.
+	 */
+	UE_API FConcertCloudFileSharingService(const FString& Role);
+
+	/** Destructor */
+	UE_API virtual ~FConcertCloudFileSharingService();
+
+	/** Make a clone of this sharing service to be used. Ownership is transferred to the caller. Callers can choose to wrap this into a unique or shared pointer. */
+	UE_API virtual IConcertFileSharingService* Clone() const override;
+
+	/** Is the file sharing service on / off. */
+	UE_API virtual bool IsEnabled() const override;
+
+	/** Turn on/off the file sharing service. */
+	UE_API virtual void SetEnabled(bool bInEnabled) override;
+
+	/** Sets the desired path to use for the service.  */
+	UE_API virtual void SetPath(const FString& Path) override;
+
+	/** Get the assigned sharing path. */
+	UE_API virtual FString GetPath() const override;
+
+	/** Delete the given InFileId if it exists. Directories are supported and if a directory is supplied then DeleteDirectoryRecursively is called. */
+	UE_API virtual void Delete(const FString& InFileId) override;
+
+	/** Copies the specified file in the service storage and return a unique ID required to consume it back. */
+	UE_API virtual bool Publish(const FString& Pathname, FString& OutFileId) override;
+
+	/** Copies the content of the archive in the service storage and return a unique ID required to consume it back. */
+	UE_API virtual bool Publish(FArchive& SrcAr, int64 Size, FString& OutFileId) override;
+
+	/** Consumes a file corresponding to the unique file ID, if available. */
+	UE_API virtual TSharedPtr<FArchive> CreateReader(const FString& InFileId) override;
+
+	/** Opens a file for writing. The file id is provided as a return argument. */
+	UE_API virtual TSharedPtr<FArchive> CreateWriter(FString& OutFileId) override;
+
+private:
+	/** Checks to see if the desired sharing path is valid and writable. */
+	bool PathIsValid() const;
+		
+	bool bEnabled = false;
+		
+	FString SharedRootPathname;
+};
+
+#undef UE_API
diff --git a/Engine/Plugins/Developer/Concert/ConcertSync/ConcertSyncCore/Source/ConcertSyncCore/Public/ConcertLocalFileSharingService.h b/Engine/Plugins/Developer/Concert/ConcertSync/ConcertSyncCore/Source/ConcertSyncCore/Public/ConcertLocalFileSharingService.h
--- a/Engine/Plugins/Developer/Concert/ConcertSync/ConcertSyncCore/Source/ConcertSyncCore/Public/ConcertLocalFileSharingService.h
+++ b/Engine/Plugins/Developer/Concert/ConcertSync/ConcertSyncCore/Source/ConcertSyncCore/Public/ConcertLocalFileSharingService.h
@@ -26,6 +26,24 @@ public:
 	UE_API FConcertLocalFileSharingService(const FString& Role);
 	UE_API virtual ~FConcertLocalFileSharingService();
 
+	/** Make a clone of this sharing service to be used. Ownership is transferred to the caller. Callers can choose to wrap this into a unique or shared pointer. */
+	UE_API virtual IConcertFileSharingService* Clone() const override;
+
+	/** Indicates if this file sharing service is enabled. Local file sharing is always enabled and it is not possible to disable at runtime. */
+	UE_API virtual bool IsEnabled() const override;
+
+	/** This is a no-op and disabling the local file sharing service is not supported.  */
+	UE_API virtual void SetEnabled(bool bInEnabled) override;
+
+	/** Sets the desired path to use for the service.  Implementors have the right to ignore the path specified.  This class does. */
+	UE_API virtual void SetPath(const FString& Path) override;
+
+	/** Get the assigned sharing path. */
+	UE_API virtual FString GetPath() const override;
+
+	/** Delete is not supported for local file sharing and only provided to satisfy interfac requirement. */
+	UE_API virtual void Delete(const FString& InFileId) override;
+
 	/** Copies the specified file in the service storage and return a unique ID required to consume it back. The stored file expires in 1 hour. */
 	UE_API virtual bool Publish(const FString& Pathname, FString& OutFileId) override;
 
@@ -35,6 +53,8 @@ public:
 	/** Consumes a file corresponding to the unique file ID, if available. The file is deleted from the service storage when the shared pointer goes out of scope. */
 	UE_API virtual TSharedPtr<FArchive> CreateReader(const FString& InFileId) override;
 
+	/** Opens a file for writing. The file id is provided as a return argument. */
+	UE_API virtual TSharedPtr<FArchive> CreateWriter(FString& OutFileId) override;
 private:
 	/** Load the list of concurrent file sharing service running across all processes. */
 	UE_API void LoadActiveServices(TArray<uint32>& OutPids);
diff --git a/Engine/Plugins/Developer/Concert/ConcertSync/ConcertSyncCore/Source/ConcertSyncCore/Public/ConcertSyncSessionFlags.h b/Engine/Plugins/Developer/Concert/ConcertSync/ConcertSyncCore/Source/ConcertSyncCore/Public/ConcertSyncSessionFlags.h
--- a/Engine/Plugins/Developer/Concert/ConcertSync/ConcertSyncCore/Source/ConcertSyncCore/Public/ConcertSyncSessionFlags.h
+++ b/Engine/Plugins/Developer/Concert/ConcertSync/ConcertSyncCore/Source/ConcertSyncCore/Public/ConcertSyncSessionFlags.h
@@ -62,7 +62,7 @@ enum class EConcertSyncSessionFlags : uint32
 	 */
 	
 	/** Default settings for a multi-user session */
-	Default_MultiUserSession = EnableLiveSync | EnableConnectionHistory | EnableLocking | EnableTransactions | EnablePackages | EnableSequencer | EnablePresence | ShouldSendTransactionSnapshots | ShouldUsePackageSandbox | EnableReplication | ShouldAllowGlobalMuting | ShouldEnableReplicationActivities | ShouldEnableRemoteEditing,
+	Default_MultiUserSession = EnableLiveSync | EnableConnectionHistory | EnableLocking | EnableTransactions | EnablePackages | EnableFileSharing | EnableSequencer | EnablePresence | ShouldSendTransactionSnapshots | ShouldUsePackageSandbox | EnableReplication | ShouldAllowGlobalMuting | ShouldEnableReplicationActivities | ShouldEnableRemoteEditing,
 	/** Default settings for a disaster recovery session */
 	Default_DisasterRecoverySession = EnableTransactions | EnablePackages | EnableFileSharing | ShouldDiscardTransactionsOnPackageUnload | ShouldSendPackagePristineState | ShouldSendPackageAutoSaves,
 };
diff --git a/Engine/Plugins/Developer/Concert/ConcertSync/ConcertSyncCore/Source/ConcertSyncCore/Public/IConcertFileSharingService.h b/Engine/Plugins/Developer/Concert/ConcertSync/ConcertSyncCore/Source/ConcertSyncCore/Public/IConcertFileSharingService.h
deleted file mode 100644
diff --git a/Engine/Plugins/Developer/Concert/ConcertSync/ConcertSyncServer/Source/ConcertSyncServer/Private/ConcertServerWorkspace.cpp b/Engine/Plugins/Developer/Concert/ConcertSync/ConcertSyncServer/Source/ConcertSyncServer/Private/ConcertServerWorkspace.cpp
--- a/Engine/Plugins/Developer/Concert/ConcertSync/ConcertSyncServer/Source/ConcertSyncServer/Private/ConcertServerWorkspace.cpp
+++ b/Engine/Plugins/Developer/Concert/ConcertSync/ConcertSyncServer/Source/ConcertSyncServer/Private/ConcertServerWorkspace.cpp
@@ -19,7 +19,7 @@
 #include "Algo/Transform.h"
 
 FConcertServerWorkspace::FConcertServerWorkspace(const TSharedRef<FConcertSyncServerLiveSession>& InLiveSession, TSharedPtr<IConcertFileSharingService> InFileSharingService)
-	: FileSharingService(MoveTemp(InFileSharingService))
+	: FileSharingService(InFileSharingService ? InFileSharingService->Clone() : nullptr)
 	, ReplicationWorkspace(
 		InLiveSession->GetSessionDatabase(),
 		UE::ConcertSyncServer::FFindSessionClient::CreateLambda([this](const FGuid& EndpointId)
@@ -81,12 +81,23 @@ void FConcertServerWorkspace::BindSession(const TSharedRef<FConcertSyncServerLiv
 	LiveSession->GetSession().RegisterCustomEventHandler<FConcertIgnoreActivityStateChangedEvent>(this, &FConcertServerWorkspace::HandleIgnoredActivityStateChanged);
 	LiveSession->GetSession().RegisterCustomEventHandler<FConcertServerLogging>(this, &FConcertServerWorkspace::HandleServerLoggingEvent);
 
+	if (FileSharingService)
+	{
+		FString BasePath = FileSharingService->GetPath();
+		FString FullPath = BasePath / LiveSession->GetSession().GetName() / TEXT("Server") / LiveSession->GetSession().GetId().ToString();
+		FileSharingService->SetPath(FullPath);
+	}
 }
 
 void FConcertServerWorkspace::UnbindSession()
 {
 	if (LiveSession.IsValid())
 	{
+		if (FileSharingService)
+		{
+			FileSharingService->Delete(FileSharingService->GetPath());
+		}
+		
 		// Destroy Sync Command Queue
 		SyncCommandQueue.Reset();
 
@@ -1218,9 +1229,14 @@ void FConcertServerWorkspace::PostActivityAdded(const int64 InActivityId)
 	}
 }
 
+bool FConcertServerWorkspace::IsFileSharingEnabled() const
+{
+	return FileSharingService && FileSharingService->IsEnabled() && EnumHasAnyFlags(LiveSession->GetSessionFlags(), EConcertSyncSessionFlags::EnableFileSharing);
+}
+
 bool FConcertServerWorkspace::CanExchangePackageDataAsByteArray(int64 PackageDataSize) const
 {
-	if (FileSharingService && EnumHasAnyFlags(LiveSession->GetSessionFlags(), EConcertSyncSessionFlags::EnableFileSharing))
+	if (IsFileSharingEnabled())
 	{
 		return FConcertPackage::ShouldEmbedPackageDataAsByteArray(PackageDataSize); // Test the package data size against a preferred limit.
 	}
@@ -1244,7 +1260,7 @@ EConcertSessionResponseCode FConcertServerWorkspace::FillPackageEvent(FConcertSy
 		OutPackageEvent.Package.PackageData.Bytes.AddUninitialized(InPackageEvent.PackageDataStream.DataSize);
 		InPackageEvent.PackageDataStream.DataAr->Serialize(OutPackageEvent.Package.PackageData.Bytes.GetData(), InPackageEvent.PackageDataStream.DataSize);
 	}
-	else if (FileSharingService)
+	else if (IsFileSharingEnabled())
 	{
 		// Share a link to the file.
 		FillPackageEventResponseCode = FileSharingService->Publish(*InPackageEvent.PackageDataStream.DataAr, InPackageEvent.PackageDataStream.DataSize, OutPackageEvent.Package.FileId) ? EConcertSessionResponseCode::Success : EConcertSessionResponseCode::Failed;
@@ -1260,8 +1276,15 @@ EConcertSessionResponseCode FConcertServerWorkspace::FillPackageEvent(FConcertSy
 
 TSharedPtr<FArchive> FConcertServerWorkspace::GetPackageDataStream(const FConcertPackage& Package)
 {
-	return !Package.FileId.IsEmpty() ?
-		FileSharingService->CreateReader(Package.FileId) : // Package data is not embedded with the activty itself, but linked.
-		MakeShared<FMemoryReader>(Package.PackageData.Bytes); // Package data is embedded in the activity object.
+	if (!Package.FileId.IsEmpty() && ensure(FileSharingService))
+	{
+		return FileSharingService->CreateReader(Package.FileId);
+	}
+	else if (Package.PackageData.Bytes.Num() > 0)
+	{
+		return MakeShared<FMemoryReader>(Package.PackageData.Bytes); // Package data is embedded in the activity object.		
+	}
+	
+	return nullptr;
 }
 
diff --git a/Engine/Plugins/Developer/Concert/ConcertSync/ConcertSyncServer/Source/ConcertSyncServer/Private/ConcertServerWorkspace.h b/Engine/Plugins/Developer/Concert/ConcertSync/ConcertSyncServer/Source/ConcertSyncServer/Private/ConcertServerWorkspace.h
--- a/Engine/Plugins/Developer/Concert/ConcertSync/ConcertSyncServer/Source/ConcertSyncServer/Private/ConcertServerWorkspace.h
+++ b/Engine/Plugins/Developer/Concert/ConcertSync/ConcertSyncServer/Source/ConcertSyncServer/Private/ConcertServerWorkspace.h
@@ -297,6 +297,9 @@ private:
 	 */
 	bool CanExchangePackageDataAsByteArray(int64 PackageDataSize) const;
 
+	/** Returns true if the file sharing service is enabled. */
+	bool IsFileSharingEnabled() const;
+		
 	/**
 	 * Write all package event fields from the event and decide if the package data should be directly embedded in the event or linked to a shared file.
 	 */
@@ -347,7 +350,7 @@ private:
 	TSet<FGuid> IgnoredActivityClients;
 
 	/** Optional side channel to exchange large blobs (package data) with the client in a scalable way (ex. the request/response transport layer is not designed and doesn't support exchanging 3GB packages). */
-	TSharedPtr<IConcertFileSharingService> FileSharingService;
+	TUniquePtr<IConcertFileSharingService> FileSharingService;
 
 	/** Manages replication data for the database. */
 	UE::ConcertSyncServer::FReplicationWorkspace ReplicationWorkspace;
diff --git a/Engine/Plugins/Developer/Concert/ConcertSync/ConcertSyncServer/Source/ConcertSyncServer/Private/ConcertSyncServer.cpp b/Engine/Plugins/Developer/Concert/ConcertSync/ConcertSyncServer/Source/ConcertSyncServer/Private/ConcertSyncServer.cpp
--- a/Engine/Plugins/Developer/Concert/ConcertSync/ConcertSyncServer/Source/ConcertSyncServer/Private/ConcertSyncServer.cpp
+++ b/Engine/Plugins/Developer/Concert/ConcertSync/ConcertSyncServer/Source/ConcertSyncServer/Private/ConcertSyncServer.cpp
@@ -284,7 +284,7 @@ IConcertServerRef FConcertSyncServer::GetConcertServer() const
 
 void FConcertSyncServer::SetFileSharingService(TSharedPtr<IConcertFileSharingService> InFileSharingService)
 {
-	FileSharingService = MoveTemp(InFileSharingService);
+	ConcertServer->SetFileSharingService(MoveTemp(InFileSharingService));
 }
 
 TOptional<FConcertSyncSessionDatabaseNonNullPtr> FConcertSyncServer::GetLiveSessionDatabase(const FGuid& SessionId)
@@ -526,7 +526,7 @@ TSharedRef<FConcertServerWorkspace> FConcertSyncServer::CreateWorkspace(const TS
 	check(InLiveSession->IsValidSession());
 	DestroyWorkspace(InLiveSession);
 	
-	const TSharedRef<FConcertServerWorkspace> Workspace = MakeShared<FConcertServerWorkspace>(InLiveSession, FileSharingService);
+	const TSharedRef<FConcertServerWorkspace> Workspace = MakeShared<FConcertServerWorkspace>(InLiveSession, ConcertServer->GetFileSharingService());
 	LiveSessionWorkspaces.Add(InLiveSession->GetSession().GetId(), Workspace);
 	return Workspace;
 }
diff --git a/Engine/Plugins/Developer/Concert/ConcertSync/ConcertSyncServer/Source/ConcertSyncServer/Private/ConcertSyncServer.h b/Engine/Plugins/Developer/Concert/ConcertSync/ConcertSyncServer/Source/ConcertSyncServer/Private/ConcertSyncServer.h
--- a/Engine/Plugins/Developer/Concert/ConcertSync/ConcertSyncServer/Source/ConcertSyncServer/Private/ConcertSyncServer.h
+++ b/Engine/Plugins/Developer/Concert/ConcertSync/ConcertSyncServer/Source/ConcertSyncServer/Private/ConcertSyncServer.h
@@ -105,9 +105,6 @@ private:
 	/** Map of archived session IDs to their associated session data */
 	TMap<FGuid, TSharedPtr<FConcertSyncServerArchivedSession>> ArchivedSessions;
 
-	/** Optional side channel to exchange large blobs (package data) with the server in a scalable way (ex. the request/response transport layer is not designed and doesn't support exchanging 3GB packages). */
-	TSharedPtr<IConcertFileSharingService> FileSharingService;
-
 	/** Logs the replication streams registered by all clients */
 	FAutoConsoleCommand LogReplicationStreamsConsoleCommand;
 	/** Logs the replication authority of the clients */
diff --git a/Engine/Plugins/Developer/Concert/ConcertSync/ConcertSyncServer/Source/ConcertSyncServer/Private/ConcertSyncServerModule.cpp b/Engine/Plugins/Developer/Concert/ConcertSync/ConcertSyncServer/Source/ConcertSyncServer/Private/ConcertSyncServerModule.cpp
--- a/Engine/Plugins/Developer/Concert/ConcertSync/ConcertSyncServer/Source/ConcertSyncServer/Private/ConcertSyncServerModule.cpp
+++ b/Engine/Plugins/Developer/Concert/ConcertSync/ConcertSyncServer/Source/ConcertSyncServer/Private/ConcertSyncServerModule.cpp
@@ -7,6 +7,13 @@
 #include "Logging/LogVerbosity.h"
 #include "Misc/EngineVersion.h"
 #include "ConcertLogGlobal.h"
+
+#include "HAL/FileManager.h"
+
+#include "Misc/Paths.h"
+
+#include "ConcertCloudFileSharingService.h"
+
 /**
  * 
  */
@@ -62,6 +69,39 @@ public:
 
 			ServerConfig->EndpointSettings.bEnableLogging |= FParse::Param(CommandLine, TEXT("CONCERTLOGGING"));
 			FParse::Bool(CommandLine, TEXT("-CONCERTLOGGING="), ServerConfig->EndpointSettings.bEnableLogging);
+
+			ServerConfig->bEnableFileSharing |= FParse::Param(CommandLine, TEXT("CONCERTFILESHARINGENABLED"));
+			FParse::Bool(CommandLine, TEXT("-CONCERTFILESHARINGENABLED="), ServerConfig->bEnableFileSharing);
+
+			FParse::Value(CommandLine, TEXT("-CONCERTFILESHAREPATH="), ServerConfig->RootFileSharingPath.Path);
+			if (ServerConfig->bEnableFileSharing)
+			{
+				bool bShouldEnableFileSharing = false;
+				if (!ServerConfig->RootFileSharingPath.Path.IsEmpty())
+				{
+					FString FullPath = FPaths::ConvertRelativePathToFull(ServerConfig->RootFileSharingPath.Path);
+					FFileStatData StatData = IFileManager::Get().GetStatData(*FullPath);
+					if (!StatData.bIsReadOnly && StatData.bIsValid && StatData.bIsDirectory)
+					{
+						ServerConfig->RootFileSharingPath.Path = FullPath;
+						bShouldEnableFileSharing = true;
+					}
+					else
+					{
+						UE_LOG(LogConcert, Warning, TEXT("The file sharing path specified was not a valid writable directory. File sharing cannot be enabled for %s."), *FullPath);
+					}
+				}
+
+				ServerConfig->bEnableFileSharing = bShouldEnableFileSharing;
+				if (bShouldEnableFileSharing)
+				{
+					UE_LOG(LogConcert, Display, TEXT("File sharing folder %s has been specified and will be used for sessions."), *ServerConfig->RootFileSharingPath.Path);
+				}
+				else
+				{
+					UE_LOG(LogConcert, Warning, TEXT("The server was launched with file sharing enabled but we were unable to validate file share path (%s)."), *ServerConfig->RootFileSharingPath.Path);
+				}
+			}
 		}
 
 		return ServerConfig;
@@ -74,6 +114,12 @@ public:
 		return ConcertSyncServer;
 	}
 
+	virtual TSharedRef<IConcertFileSharingService> CreateFileSharingService(const FString& InRole) override
+	{
+		TSharedRef<IConcertFileSharingService> FileShare = MakeShared<FConcertCloudFileSharingService>(InRole);
+		return FileShare;
+	}
+
 	virtual FOnConcertSyncServerCreated& OnServerCreated() override
 	{
 		return OnConcertSyncServerCreatedDelegate;
diff --git a/Engine/Plugins/Developer/Concert/ConcertSync/ConcertSyncServer/Source/ConcertSyncServer/Public/ConcertSyncServerLoop.inl b/Engine/Plugins/Developer/Concert/ConcertSync/ConcertSyncServer/Source/ConcertSyncServer/Public/ConcertSyncServerLoop.inl
--- a/Engine/Plugins/Developer/Concert/ConcertSync/ConcertSyncServer/Source/ConcertSyncServer/Public/ConcertSyncServerLoop.inl
+++ b/Engine/Plugins/Developer/Concert/ConcertSync/ConcertSyncServer/Source/ConcertSyncServer/Public/ConcertSyncServerLoop.inl
@@ -142,7 +142,7 @@ int32 ConcertSyncServerLoop(const TCHAR* CommandLine, const FConcertSyncServerLo
 
 	// Setup Concert Sync to run in server mode.
 	TSharedPtr<IConcertSyncServer> ConcertSyncServer = IConcertSyncServerModule::Get().CreateServer(InitArgs.ServiceRole, InitArgs.ServiceAutoArchiveSessionFilter);
-	ConcertSyncServer->SetFileSharingService(InitArgs.FileSharingService);
+	ConcertSyncServer->SetFileSharingService(InitArgs.GetFileSharingServiceFunc ? InitArgs.GetFileSharingServiceFunc() : InitArgs.FileSharingService);
 	ConcertSyncServer->Startup(ServerConfig, InitArgs.SessionFlags);
 
 	// if we have a default session, set it up properly
diff --git a/Engine/Plugins/Developer/Concert/ConcertSync/ConcertSyncServer/Source/ConcertSyncServer/Public/ConcertSyncServerLoopInitArgs.h b/Engine/Plugins/Developer/Concert/ConcertSync/ConcertSyncServer/Source/ConcertSyncServer/Public/ConcertSyncServerLoopInitArgs.h
--- a/Engine/Plugins/Developer/Concert/ConcertSync/ConcertSyncServer/Source/ConcertSyncServer/Public/ConcertSyncServerLoopInitArgs.h
+++ b/Engine/Plugins/Developer/Concert/ConcertSync/ConcertSyncServer/Source/ConcertSyncServer/Public/ConcertSyncServerLoopInitArgs.h
@@ -34,6 +34,9 @@ struct FConcertSyncServerLoopInitArgs
 	/** The optional file sharing server to exchange large files. Can be null. */
 	TSharedPtr<IConcertFileSharingService> FileSharingService;
 
+	/** Function to get the file sharing service during initialization of concert server. f this is defined then it takes precedence over FileSharingService member. */
+	TFunction<TSharedPtr<IConcertFileSharingService>()> GetFileSharingServiceFunc;
+	
 	/** Function to get the server settings object to configure the server with with, or unbound to parse the default settings */
 	TFunction<const UConcertServerConfig*()> GetServerConfigFunc;
 
diff --git a/Engine/Plugins/Developer/Concert/ConcertSync/ConcertSyncServer/Source/ConcertSyncServer/Public/IConcertSyncServerModule.h b/Engine/Plugins/Developer/Concert/ConcertSync/ConcertSyncServer/Source/ConcertSyncServer/Public/IConcertSyncServerModule.h
--- a/Engine/Plugins/Developer/Concert/ConcertSync/ConcertSyncServer/Source/ConcertSyncServer/Public/IConcertSyncServerModule.h
+++ b/Engine/Plugins/Developer/Concert/ConcertSync/ConcertSyncServer/Source/ConcertSyncServer/Public/IConcertSyncServerModule.h
@@ -3,11 +3,13 @@
 #pragma once
 
 #include "CoreMinimal.h"
+#include "IConcertSyncServer.h"
 #include "Modules/ModuleInterface.h"
 #include "Modules/ModuleManager.h"
 
 class IConcertSyncServer;
 class UConcertServerConfig;
+class IConcertFileSharingService;
 
 struct FConcertSessionFilter;
 
@@ -57,6 +59,13 @@ public:
 	 */
 	virtual TSharedRef<IConcertSyncServer> CreateServer(const FString& InRole, const FConcertSessionFilter& InAutoArchiveSessionFilter) = 0;
 
+	/**
+	 * Creates a shared reference to the file sharing service to be used by concert.
+	 *
+	 * @return a non-null shared reference to the file sharing service. 
+	 */
+	virtual TSharedRef<IConcertFileSharingService> CreateFileSharingService(const FString& InRole) = 0;
+		
 	/**
 	 * Callback for when a ConcertSyncServer has been created.
 	 */
diff --git a/Engine/Plugins/VirtualProduction/Switchboard/Source/Switchboard/switchboard/config.py b/Engine/Plugins/VirtualProduction/Switchboard/Source/Switchboard/switchboard/config.py
--- a/Engine/Plugins/VirtualProduction/Switchboard/Source/Switchboard/switchboard/config.py
+++ b/Engine/Plugins/VirtualProduction/Switchboard/Source/Switchboard/switchboard/config.py
@@ -2910,6 +2910,16 @@ class Config(object):
                 "Directory for Live Sessions",
                 data.get('muserver_working_dir', '')
             ),
+            "muserver_filesharing_srv": DirectoryPathSetting(
+                "muserver_filesharing_srv",
+                "Directory or URI for File Sharing Service",
+                data.get('muserver_filesharing_srv', '')
+            ),
+            "muserver_filesharing_enabled": BoolSetting(
+                "muserver_filesharing_enabled",
+                "Use File Sharing Service",
+                data.get('muserver_filesharing_enabled', False)
+            ),
             "muserver_auto_launch": BoolSetting(
                 "muserver_auto_launch",
                 "Auto Launch",
@@ -2963,7 +2973,9 @@ class Config(object):
         self.MUSERVER_WORKING_DIR = self.mu_settings["muserver_working_dir"]
         self.MUSERVER_ARCHIVE_DIR = self.mu_settings["muserver_archive_dir"]
         self.MUSERVER_SESSION_NAME = self.mu_settings["muserver_session_name"]
-
+        self.MUSERVER_FILESHARING_SRV = self.mu_settings["muserver_filesharing_srv"]
+        self.MUSERVER_FILESHARING_ENABLED = self.mu_settings["muserver_filesharing_enabled"]
+        
     def save_unreal_insights(self, data):
         data['tracing_enabled'] = self.INSIGHTS_TRACE_ENABLE.get_value()
         data['tracing_args'] = self.INSIGHTS_TRACE_ARGS.get_value()
@@ -2985,7 +2997,10 @@ class Config(object):
         data["muserver_archive_dir"] = self.MUSERVER_ARCHIVE_DIR.get_value()
         data["muserver_working_dir"] = self.MUSERVER_WORKING_DIR.get_value()
         data["muserver_session_name"] = self.MUSERVER_SESSION_NAME.get_value()
+        data["muserver_filesharing_srv"] = self.MUSERVER_FILESHARING_SRV.get_value() 
+        data["muserver_filesharing_enabled"] = self.MUSERVER_FILESHARING_ENABLED.get_value() 
 
+        
     def load_plugin_settings(self, device_type, settings):
         ''' Updates plugin settings values with those read from the config file.
         '''
diff --git a/Engine/Plugins/VirtualProduction/Switchboard/Source/Switchboard/switchboard/devices/ndisplay/plugin_ndisplay.py b/Engine/Plugins/VirtualProduction/Switchboard/Source/Switchboard/switchboard/devices/ndisplay/plugin_ndisplay.py
--- a/Engine/Plugins/VirtualProduction/Switchboard/Source/Switchboard/switchboard/devices/ndisplay/plugin_ndisplay.py
+++ b/Engine/Plugins/VirtualProduction/Switchboard/Source/Switchboard/switchboard/devices/ndisplay/plugin_ndisplay.py
@@ -1287,6 +1287,13 @@ class DevicenDisplay(DeviceUnreal):
             args.extend([
                 '-CONCERTRETRYAUTOCONNECTONERROR',
                 '-CONCERTAUTOCONNECT'])
+
+        if CONFIG.MUSERVER_FILESHARING_ENABLED.get_value() and CONFIG.MUSERVER_FILESHARING_SRV.get_value():
+            filesharing = CONFIG.MUSERVER_FILESHARING_SRV.get_value()
+            args.extend([
+                '-CONCERTFILESHARINGENABLED=1',
+                f'-CONCERTFILESHAREPATH={filesharing}'])
+
         mu_server = switchboard_application.get_multi_user_server_instance()
         args.extend([
             f'-CONCERTSERVER="{mu_server.configured_server_name()}"',
diff --git a/Engine/Plugins/VirtualProduction/Switchboard/Source/Switchboard/switchboard/devices/unreal/plugin_unreal.py b/Engine/Plugins/VirtualProduction/Switchboard/Source/Switchboard/switchboard/devices/unreal/plugin_unreal.py
--- a/Engine/Plugins/VirtualProduction/Switchboard/Source/Switchboard/switchboard/devices/unreal/plugin_unreal.py
+++ b/Engine/Plugins/VirtualProduction/Switchboard/Source/Switchboard/switchboard/devices/unreal/plugin_unreal.py
@@ -2449,6 +2449,12 @@ class DeviceUnreal(Device):
                 '-CONCERTAUTOCONNECT'
             ])
 
+        if CONFIG.MUSERVER_FILESHARING_ENABLED.get_value() and CONFIG.MUSERVER_FILESHARING_SRV.get_value():
+            filesharing = CONFIG.MUSERVER_FILESHARING_SRV.get_value()
+            args.extend([
+                '-CONCERTFILESHARINGENABLED=1',
+                f'-CONCERTFILESHAREPATH={filesharing}'])
+
         mu_server = switchboard_application.get_multi_user_server_instance()
         args.extend([
             f'-CONCERTSERVER="{mu_server.configured_server_name()}"',
diff --git a/Engine/Plugins/VirtualProduction/Switchboard/Source/Switchboard/switchboard/switchboard_application.py b/Engine/Plugins/VirtualProduction/Switchboard/Source/Switchboard/switchboard/switchboard_application.py
--- a/Engine/Plugins/VirtualProduction/Switchboard/Source/Switchboard/switchboard/switchboard_application.py
+++ b/Engine/Plugins/VirtualProduction/Switchboard/Source/Switchboard/switchboard/switchboard_application.py
@@ -218,6 +218,11 @@ class MultiUserApplication:
             if CONFIG.MUSERVER_CLEAN_HISTORY.get_value():
                 args.append('-ConcertClean')
 
+            if CONFIG.MUSERVER_FILESHARING_ENABLED.get_value() and CONFIG.MUSERVER_FILESHARING_SRV.get_value():
+                args.append('-CONCERTFILESHARINGENABLED=1')
+                filesharing = CONFIG.MUSERVER_FILESHARING_SRV.get_value()
+                args.append(f'-CONCERTFILESHAREPATH={filesharing}')
+                
             args = list(filter(None, args))
 
             launch_cmd_full = ' '.join(args)
diff --git a/Engine/Source/Programs/UnrealMultiUserServer/Private/UnrealMultiUserServerRun.cpp b/Engine/Source/Programs/UnrealMultiUserServer/Private/UnrealMultiUserServerRun.cpp
--- a/Engine/Source/Programs/UnrealMultiUserServer/Private/UnrealMultiUserServerRun.cpp
+++ b/Engine/Source/Programs/UnrealMultiUserServer/Private/UnrealMultiUserServerRun.cpp
@@ -14,7 +14,13 @@ int32 RunUnrealMultiUserServer(int ArgC, TCHAR* ArgV[])
 	ServerLoopInitArgs.SessionFlags = EConcertSyncSessionFlags::Default_MultiUserSession;
 	ServerLoopInitArgs.ServiceRole = Role;
 	ServerLoopInitArgs.ServiceFriendlyName = TEXT("Multi-User Editing Server");
-
+	ServerLoopInitArgs.GetFileSharingServiceFunc = [Role]() -> TSharedPtr<IConcertFileSharingService>
+	{
+		return IConcertSyncServerModule::Get().CreateFileSharingService(Role);
+	};
+	
+	ServerLoopInitArgs.FileSharingService = IConcertSyncServerModule::Get().CreateFileSharingService(Role);
+	
 	ServerLoopInitArgs.GetServerConfigFunc = [Role]() -> const UConcertServerConfig*
 	{
 		UConcertServerConfig* ServerConfig = IConcertSyncServerModule::Get().ParseServerSettings(FCommandLine::Get());
diff --git a/Engine/Source/Programs/UnrealMultiUserSlateServer/Private/UnrealMultiUserSlateServerRun.cpp b/Engine/Source/Programs/UnrealMultiUserSlateServer/Private/UnrealMultiUserSlateServerRun.cpp
--- a/Engine/Source/Programs/UnrealMultiUserSlateServer/Private/UnrealMultiUserSlateServerRun.cpp
+++ b/Engine/Source/Programs/UnrealMultiUserSlateServer/Private/UnrealMultiUserSlateServerRun.cpp
@@ -4,6 +4,7 @@
 
 #include "ConcertSettings.h"
 #include "ConcertSyncServerLoop.h"
+
 #include "IMultiUserServerModule.h"
 
 namespace UE::UnrealMultiUserServer
@@ -33,6 +34,10 @@ int32 RunUnrealMultiUserServer(const TCHAR* CommandLine)
 	ServerLoopInitArgs.SessionFlags = EConcertSyncSessionFlags::Default_MultiUserSession;
 	ServerLoopInitArgs.ServiceRole = Role;
 	ServerLoopInitArgs.ServiceFriendlyName = TEXT("Multi-User Editing Server");
+	ServerLoopInitArgs.GetFileSharingServiceFunc = [Role]() -> TSharedPtr<IConcertFileSharingService>
+	{
+		return IConcertSyncServerModule::Get().CreateFileSharingService(Role);
+	};
 
 	ServerLoopInitArgs.GetServerConfigFunc = [Role]() -> const UConcertServerConfig*
 	{
