Specify FConcertSessionChallengeData::ChallengeKey should be ignored when running member initialization test.

#Synced-CL 49070374

diff --git a/Engine/Plugins/Developer/Concert/ConcertMain/Source/Concert/Public/ConcertMessageData.h b/Engine/Plugins/Developer/Concert/ConcertMain/Source/Concert/Public/ConcertMessageData.h
--- a/Engine/Plugins/Developer/Concert/ConcertMain/Source/Concert/Public/ConcertMessageData.h
+++ b/Engine/Plugins/Developer/Concert/ConcertMain/Source/Concert/Public/ConcertMessageData.h
@@ -257,7 +257,7 @@ struct FConcertSessionChallengeData
 
 	bool operator==(const FConcertSessionChallengeData& Other) const = default;
 	
-	UPROPERTY()
+	UPROPERTY(Meta = (IgnoreForMemberInitializationTest))
 	FGuid ChallengeKey;
 
 	UPROPERTY()
